Functional testing improvements #UE4

- a functional test can specify an "observation point" where player will get teleported when the test is started. I suggest using CameraActors
- made function test maps be executed as regular games when run from automation in Editor. Running as "simulation" was not really working.

diff --git a/Engine/Source/Developer/FunctionalTesting/Classes/FunctionalTest.h b/Engine/Source/Developer/FunctionalTesting/Classes/FunctionalTest.h
--- a/Engine/Source/Developer/FunctionalTesting/Classes/FunctionalTest.h
+++ b/Engine/Source/Developer/FunctionalTesting/Classes/FunctionalTest.h
@@ -44,6 +44,9 @@ class AFunctionalTest : public AActor
 	UPROPERTY(EditAnywhere, BlueprintReadWrite, Category=FunctionalTesting)
 	FText TimesUpMessage;
 
+	UPROPERTY(EditAnywhere, BlueprintReadWrite, Category=FunctionalTesting)
+	AActor* ObservationPoint;
+
 	/** Called when the test is started */
 	UPROPERTY(BlueprintAssignable)
 	FFunctionalTestEventSignature OnTestStart;
@@ -54,7 +57,7 @@ class AFunctionalTest : public AActor
 
 	UPROPERTY(Transient)
 	TArray<AActor*> AutoDestroyActors;
-
+	
 	FString FailureMessage;
 	
 #if WITH_EDITORONLY_DATA
@@ -104,6 +107,10 @@ public:
 	virtual void Tick(float DeltaSeconds) OVERRIDE;
 	// AActor interface end
 
+protected:
+	void GoToObservationPoint();
+
+public:
 	FFunctionalTestDoneSignature TestFinishedObserver;
 
 protected:
diff --git a/Engine/Source/Developer/FunctionalTesting/Private/FunctionalTest.cpp b/Engine/Source/Developer/FunctionalTesting/Private/FunctionalTest.cpp
--- a/Engine/Source/Developer/FunctionalTesting/Private/FunctionalTest.cpp
+++ b/Engine/Source/Developer/FunctionalTesting/Private/FunctionalTest.cpp
@@ -81,6 +81,8 @@ bool AFunctionalTest::StartTest()
 	}
 
 	bIsRunning = true;
+
+	GoToObservationPoint();
 	
 	OnTestStart.Broadcast();
 
@@ -91,6 +93,12 @@ void AFunctionalTest::FinishTest(TEnumAsByte<EFunctionalTestResult::Type> TestRe
 {
 	const static UEnum* FTestResultTypeEnum = FindObject<UEnum>( NULL, TEXT("FunctionalTesting.FunctionalTest.EFunctionalTestResult") );
 	
+	if (bIsRunning == false)
+	{
+		// ignore
+		return;
+	}
+
 	bIsRunning = false;
 	SetActorTickEnabled(false);
 
@@ -222,3 +230,22 @@ void AFunctionalTest::PostEditChangeProperty( struct FPropertyChangedEvent& Prop
 }
 
 #endif // WITH_EDITOR
+
+void AFunctionalTest::GoToObservationPoint()
+{
+	if (ObservationPoint == NULL)
+	{
+		return;
+	}
+
+	UWorld* World = GetWorld();
+	if (World)
+	{
+		APlayerController* PC = World->GetFirstPlayerController();
+		if (PC && PC->GetPawn())
+		{
+			PC->GetPawn()->TeleportTo(ObservationPoint->GetActorLocation(), ObservationPoint->GetActorRotation());
+			PC->SetControlRotation(ObservationPoint->GetActorRotation());
+		}
+	}
+}
\ No newline at end of file
diff --git a/Engine/Source/Developer/FunctionalTesting/Private/Tests/FuncTestPerforming.cpp b/Engine/Source/Developer/FunctionalTesting/Private/Tests/FuncTestPerforming.cpp
--- a/Engine/Source/Developer/FunctionalTesting/Private/Tests/FuncTestPerforming.cpp
+++ b/Engine/Source/Developer/FunctionalTesting/Private/Tests/FuncTestPerforming.cpp
@@ -159,11 +159,7 @@ bool FFunctionalTestingMaps::RunTest(const FString& Parameters)
 		bool bLoadAsTemplate = false;
 		bool bShowProgress = false;
 		FEditorFileUtils::LoadMap(MapName, bLoadAsTemplate, bShowProgress);
-
-		const bool bRunInPIE = MapName.EndsWith(TEXT("_PIE"));
-		UE_LOG(LogFunctionalTesting, Display, TEXT("Running map %s in %s mode"), *MapName, bRunInPIE ? TEXT("PIE") : TEXT("SIMULATION"));
-
-		EditorEngine->PlayInEditor( GWorld, bRunInPIE );
+		EditorEngine->PlayInEditor(GWorld, /*bInSimulateInEditor=*/false);
 	}
 	else
 #endif // WITH_EDITOR
