Improved FunctionalTesting logging to properly report in Automation #UE4

diff --git a/Engine/Source/Developer/FunctionalTesting/Classes/FunctionalTest.h b/Engine/Source/Developer/FunctionalTesting/Classes/FunctionalTest.h
--- a/Engine/Source/Developer/FunctionalTesting/Classes/FunctionalTest.h
+++ b/Engine/Source/Developer/FunctionalTesting/Classes/FunctionalTest.h
@@ -107,6 +107,8 @@ public:
 	virtual void Tick(float DeltaSeconds) override;
 	// AActor interface end
 
+	bool IsSuccessful() const { Result == EFunctionalTestResult::Succeeded; }
+
 protected:
 	void GoToObservationPoint();
 
diff --git a/Engine/Source/Developer/FunctionalTesting/Classes/FunctionalTestingManager.h b/Engine/Source/Developer/FunctionalTesting/Classes/FunctionalTestingManager.h
--- a/Engine/Source/Developer/FunctionalTesting/Classes/FunctionalTestingManager.h
+++ b/Engine/Source/Developer/FunctionalTesting/Classes/FunctionalTestingManager.h
@@ -2,6 +2,7 @@
 
 #pragma once
 
+#include "Misc/AutomationTest.h"
 #include "FunctionalTestingManager.generated.h"
 
 UCLASS(BlueprintType, MinimalAPI)
@@ -21,7 +22,7 @@ class UFunctionalTestingManager : public UBlueprintFunctionLibrary
 	UFUNCTION(BlueprintCallable, Category="FunctionalTesting", meta=(HidePin="WorldContext", DefaultToSelf="WorldContext" ) )
 	/** Triggers in sequence all functional tests found on the level.
 	 *	@return true if any tests have been triggered */
-	static bool RunAllFunctionalTests(UObject* WorldContext, bool bNewLog=true, bool bRunLooped=false, bool bWaitForNavigationBuildFinish=true);
+	static bool RunAllFunctionalTests(UObject* WorldContext, bool bNewLog=true, bool bRunLooped=false, bool bWaitForNavigationBuildFinish=true, bool bLoopOnlyFailedTests = true);
 		
 	bool IsRunning() const { return bIsRunning; }
 	bool IsLooped() const { return bLooped; }
@@ -29,6 +30,14 @@ class UFunctionalTestingManager : public UBlueprintFunctionLibrary
 
 	void TickMe(float DeltaTime);
 
+	//----------------------------------------------------------------------//
+	// Automation logging
+	//----------------------------------------------------------------------//
+	static void SetAutomationExecutionInfo(FAutomationTestExecutionInfo* InExecutionInfo) { ExecutionInfo = InExecutionInfo; }
+	static void AddError(const FText& InError);
+	static void AddWarning(const FText& InWarning);
+	static void AddLogItem(const FText& InLogItem);
+	
 private:
 	void LogMessage(const FString& MessageString, TSharedPtr<IMessageLogListing> LogListing = NULL);
 	
@@ -49,5 +58,8 @@ protected:
 	bool bLooped;
 	bool bWaitForNavigationBuildFinish;
 	bool bInitialDelayApplied;
+	bool bDiscardSuccessfulTests;
 	uint32 CurrentIteration;
+
+	static FAutomationTestExecutionInfo* ExecutionInfo;
 };
diff --git a/Engine/Source/Developer/FunctionalTesting/Private/FunctionalTest.cpp b/Engine/Source/Developer/FunctionalTesting/Private/FunctionalTest.cpp
--- a/Engine/Source/Developer/FunctionalTesting/Private/FunctionalTest.cpp
+++ b/Engine/Source/Developer/FunctionalTesting/Private/FunctionalTest.cpp
@@ -146,12 +146,26 @@ void AFunctionalTest::FinishTest(TEnumAsByte<EFunctionalTestResult::Type> TestRe
 			break;
 	}
 
-	FMessageLog("FunctionalTestingLog").Message(MessageLogSeverity, FText::FromString(GetActorLabel()))
+	TSharedRef<FTokenizedMessage> LogMessage = FTokenizedMessage::Create(MessageLogSeverity, FText::FromString(GetActorLabel()))
 		->AddToken(FTextToken::Create(ResultText))
 		->AddToken(FTextToken::Create(FText::FromString(Message)))
 		->AddToken(FTextToken::Create(FText::FromString(AdditionalDetails)))
 		->AddToken(FTextToken::Create(FText::FromString(FailureMessage)));
 
+	switch (MessageLogSeverity)
+	{
+	case EMessageSeverity::Error:
+		UFunctionalTestingManager::AddError(LogMessage->ToText());
+		break;
+	case EMessageSeverity::Warning:
+		UFunctionalTestingManager::AddWarning(LogMessage->ToText());
+		break;
+	default:
+		UFunctionalTestingManager::AddLogItem(LogMessage->ToText());
+		break;
+	}
+
+
 	TestFinishedObserver.ExecuteIfBound(this);
 }
 
diff --git a/Engine/Source/Developer/FunctionalTesting/Private/FunctionalTestingManager.cpp b/Engine/Source/Developer/FunctionalTesting/Private/FunctionalTestingManager.cpp
--- a/Engine/Source/Developer/FunctionalTesting/Private/FunctionalTestingManager.cpp
+++ b/Engine/Source/Developer/FunctionalTesting/Private/FunctionalTestingManager.cpp
@@ -34,6 +34,9 @@
 //----------------------------------------------------------------------//
 // 
 //----------------------------------------------------------------------//
+
+FAutomationTestExecutionInfo* UFunctionalTestingManager::ExecutionInfo = NULL;
+
 UFunctionalTestingManager::UFunctionalTestingManager( const class FPostConstructInitializeProperties& PCIP )
 	: Super(PCIP)
 	, bIsRunning(false)
@@ -49,7 +52,7 @@ void UFunctionalTestingManager::SetUpTests()
 	OnSetupTests.Broadcast();
 }
 
-bool UFunctionalTestingManager::RunAllFunctionalTests(UObject* WorldContext, bool bNewLog, bool bRunLooped, bool bInWaitForNavigationBuildFinish)
+bool UFunctionalTestingManager::RunAllFunctionalTests(UObject* WorldContext, bool bNewLog, bool bRunLooped, bool bInWaitForNavigationBuildFinish, bool bLoopOnlyFailedTests)
 {
 	UFunctionalTestingManager* Manager = GetManager(WorldContext);
 
@@ -69,6 +72,7 @@ bool UFunctionalTestingManager::RunAllFunctionalTests(UObject* WorldContext, boo
 	Manager->bLooped = bRunLooped;
 	Manager->bWaitForNavigationBuildFinish = bInWaitForNavigationBuildFinish;
 	Manager->CurrentIteration = 0;
+	Manager->bDiscardSuccessfulTests = bLoopOnlyFailedTests;
 
 	for (TActorIterator<AFunctionalTest> It(GWorld); It; ++It)
 	{
@@ -89,7 +93,7 @@ bool UFunctionalTestingManager::RunAllFunctionalTests(UObject* WorldContext, boo
 
 	if (Manager->bIsRunning == false)
 	{
-		FunctionalTestingLog.Info(LOCTEXT("NoTestsDefined", "No tests defined on map or . DONE."));		
+		AddLogItem(LOCTEXT("NoTestsDefined", "No tests defined on map or . DONE.")); 
 		return false;
 	}
 
@@ -150,11 +154,13 @@ void UFunctionalTestingManager::OnTestDone(AFunctionalTest* FTest)
 
 void UFunctionalTestingManager::NotifyTestDone(AFunctionalTest* FTest)
 {
-	FMessageLog FunctionalTestingLog("FunctionalTestingLog");
-
 	if (FTest->WantsToRunAgain() == false)
 	{
 		TestsLeft.RemoveSingle(FTest);
+		/*if (bDiscardSuccessfulTests && FTest->IsSuccessful())
+		{
+			AllTests.RemoveSingle(FTest);
+		}*/
 		FTest->CleanUp();
 	}
 
@@ -175,12 +181,12 @@ void UFunctionalTestingManager::NotifyTestDone(AFunctionalTest* FTest)
 			TestsLeft = AllTests;
 			FFormatNamedArguments Arguments;
 			Arguments.Add(TEXT("IterationIndex"), CurrentIteration);
-			FunctionalTestingLog.Info( FText::Format(LOCTEXT("StartingIteration", "----- Starting iteration {IterationIndex} -----"), Arguments) );
+			AddLogItem(FText::Format(LOCTEXT("StartingIteration", "----- Starting iteration {IterationIndex} -----"), Arguments));
 			bIsRunning = RunFirstValidTest();
 		}
 		else
 		{
-			FunctionalTestingLog.Info( LOCTEXT("TestDone", "DONE.") );
+			AddLogItem(LOCTEXT("TestDone", "DONE."));
 			RemoveFromRoot();
 		}
 	}
@@ -220,4 +226,38 @@ void UFunctionalTestingManager::TickMe(float DeltaTime)
 	
 }
 
+//----------------------------------------------------------------------//
+// logging
+//----------------------------------------------------------------------//
+void UFunctionalTestingManager::AddError(const FText& InError)
+{
+	FMessageLog FunctionalTestingLog("FunctionalTestingLog");
+	FunctionalTestingLog.Error(InError);
+	if (ExecutionInfo)
+	{
+		ExecutionInfo->Errors.Add(InError.ToString());
+	}
+}
+
+
+void UFunctionalTestingManager::AddWarning(const FText& InWarning)
+{
+	FMessageLog FunctionalTestingLog("FunctionalTestingLog");
+	FunctionalTestingLog.Warning(InWarning);
+	if (ExecutionInfo)
+	{
+		ExecutionInfo->Warnings.Add(InWarning.ToString());
+	}
+}
+
+
+void UFunctionalTestingManager::AddLogItem(const FText& InLogItem)
+{
+	FMessageLog FunctionalTestingLog("FunctionalTestingLog");
+	FunctionalTestingLog.Info(InLogItem);
+	if (ExecutionInfo)
+	{
+		ExecutionInfo->LogItems.Add(InLogItem.ToString());
+	}
+}
 #undef LOCTEXT_NAMESPACE
diff --git a/Engine/Source/Developer/FunctionalTesting/Private/Tests/FuncTestPerforming.cpp b/Engine/Source/Developer/FunctionalTesting/Private/Tests/FuncTestPerforming.cpp
--- a/Engine/Source/Developer/FunctionalTesting/Private/Tests/FuncTestPerforming.cpp
+++ b/Engine/Source/Developer/FunctionalTesting/Private/Tests/FuncTestPerforming.cpp
@@ -2,7 +2,6 @@
 
 #include "FunctionalTestingPrivatePCH.h"
 #if WITH_EDITOR
-//#include "UnrealEd.h"
 #include "Editor/UnrealEd/Classes/Editor/EditorEngine.h"
 #include "Editor/UnrealEd/Public/FileHelpers.h"
 #endif // WITH_EDITOR
@@ -97,13 +96,30 @@ bool FStartFTestsOnMap::Update()
 /**
  * 
  */
-IMPLEMENT_COMPLEX_AUTOMATION_TEST( FFunctionalTestingMaps, "Maps.Functional Testing", (EAutomationTestFlags::ATF_Game | EAutomationTestFlags::ATF_Editor) )
 
+// create test base class
+IMPLEMENT_COMPLEX_AUTOMATION_TEST_PRIVATE(FFunctionalTestingMapsBase, "Maps.Functional Testing", (EAutomationTestFlags::ATF_Game | EAutomationTestFlags::ATF_Editor))
+// implement specific class with non-standard overrides
+class FFunctionalTestingMaps : public FFunctionalTestingMapsBase
+{
+public:
+	FFunctionalTestingMaps(const FString& InName) : FFunctionalTestingMapsBase(InName)
+	{}
+	virtual ~FFunctionalTestingMaps()
+	{
+		UFunctionalTestingManager::SetAutomationExecutionInfo(NULL);
+	}
+};
+
+namespace
+{
+	FFunctionalTestingMaps FFunctionalTestingMapsAutomationTestInstance(TEXT("FFunctionalTestingMaps"));
+}
 
 /** 
  * Requests a enumeration of all maps to be loaded
  */
-void FFunctionalTestingMaps::GetTests(TArray<FString>& OutBeautifiedNames, TArray <FString>& OutTestCommands) const
+void FFunctionalTestingMapsBase::GetTests(TArray<FString>& OutBeautifiedNames, TArray <FString>& OutTestCommands) const
 {
 	TArray<FString> FileList;
 #if WITH_EDITOR
@@ -123,7 +139,7 @@ void FFunctionalTestingMaps::GetTests(TArray<FString>& OutBeautifiedNames, TArra
 #endif
 
 	// Iterate over all files, adding the ones with the map extension..
-	for( int32 FileIndex=0; FileIndex< FileList.Num(); FileIndex++ )
+	for( int32 FileIndex = 0; FileIndex< FileList.Num(); FileIndex++ )
 	{
 		const FString& Filename = FileList[FileIndex];
 
@@ -149,10 +165,14 @@ void FFunctionalTestingMaps::GetTests(TArray<FString>& OutBeautifiedNames, TArra
  * @param Parameters - Should specify which map name to load
  * @return	TRUE if the test was successful, FALSE otherwise
  */
-bool FFunctionalTestingMaps::RunTest(const FString& Parameters)
+bool FFunctionalTestingMapsBase::RunTest(const FString& Parameters)
 {
 	FString MapName = Parameters;
 
+	SetSuppressLogs(true);
+	ExecutionInfo.Clear();
+	UFunctionalTestingManager::SetAutomationExecutionInfo(&ExecutionInfo);
+
 #if WITH_EDITOR
 	if (GIsEditor)
 	{
diff --git a/Engine/Source/Runtime/Core/Public/Misc/AutomationTest.h b/Engine/Source/Runtime/Core/Public/Misc/AutomationTest.h
--- a/Engine/Source/Runtime/Core/Public/Misc/AutomationTest.h
+++ b/Engine/Source/Runtime/Core/Public/Misc/AutomationTest.h
@@ -1084,13 +1084,13 @@ public: \
 		virtual bool IsStressTest() const { return false; } \
 		virtual uint32 GetRequiredDeviceNum() const { return 1; } \
 	protected: \
-		virtual void GetTests(TArray<FString>& OutBeautifiedNames, TArray <FString>& OutTestCommands) const \
+		virtual void GetTests(TArray<FString>& OutBeautifiedNames, TArray <FString>& OutTestCommands) const override \
 		{ \
 			OutBeautifiedNames.Add(PrettyName); \
 			OutTestCommands.Add(FString()); \
 		} \
-		virtual bool RunTest(const FString& Parameters); \
-		virtual FString GetBeautifiedTestName() const { return PrettyName; } \
+		virtual bool RunTest(const FString& Parameters) override; \
+		virtual FString GetBeautifiedTestName() const override { return PrettyName; } \
 	};
 
 #define IMPLEMENT_COMPLEX_AUTOMATION_TEST_PRIVATE( TClass, PrettyName, TFlags ) \
@@ -1104,8 +1104,8 @@ public: \
 		virtual uint32 GetRequiredDeviceNum() const { return 1; } \
 	protected: \
 		virtual void GetTests(TArray<FString>& OutBeautifiedNames, TArray <FString>& OutTestCommands) const override; \
-		virtual bool RunTest(const FString& Parameters); \
-		virtual FString GetBeautifiedTestName() const { return PrettyName; } \
+		virtual bool RunTest(const FString& Parameters) override; \
+		virtual FString GetBeautifiedTestName() const override { return PrettyName; } \
 	};
 
 #define IMPLEMENT_NETWORKED_AUTOMATION_TEST_PRIVATE(TClass, PrettyName, TFlags, NumParticipants) \
@@ -1117,13 +1117,13 @@ public: \
 		virtual uint32 GetTestFlags() const { return (TFlags & ~(EAutomationTestFlags::ATF_Editor | EAutomationTestFlags::ATF_Commandlet | EAutomationTestFlags::ATF_SmokeTest)); } \
 		virtual uint32 GetRequiredDeviceNum() const { return NumParticipants; } \
 	protected: \
-		virtual void GetTests(TArray<FString>& OutBeautifiedNames, TArray <FString>& OutTestCommands) const \
+		virtual void GetTests(TArray<FString>& OutBeautifiedNames, TArray <FString>& OutTestCommands) const override \
 		{ \
 			OutBeautifiedNames.Add(PrettyName); \
 			OutTestCommands.Add(FString()); \
 		} \
-		virtual bool RunTest(const FString& Parameters); \
-		virtual FString GetBeautifiedTestName() const { return PrettyName; } \
+		virtual bool RunTest(const FString& Parameters) override; \
+		virtual FString GetBeautifiedTestName() const override { return PrettyName; } \
 	};
 
 
