Fix compile error in CL# 2106961

diff --git a/Engine/Source/Developer/FunctionalTesting/Classes/FunctionalTest.h b/Engine/Source/Developer/FunctionalTesting/Classes/FunctionalTest.h
--- a/Engine/Source/Developer/FunctionalTesting/Classes/FunctionalTest.h
+++ b/Engine/Source/Developer/FunctionalTesting/Classes/FunctionalTest.h
@@ -107,7 +107,7 @@ public:
 	virtual void Tick(float DeltaSeconds) override;
 	// AActor interface end
 
-	bool IsSuccessful() const { Result == EFunctionalTestResult::Succeeded; }
+	bool IsSuccessful() const { return Result == EFunctionalTestResult::Succeeded; }
 
 protected:
 	void GoToObservationPoint();
