Added a way to trigger specific subset of functional tests #UE4

When a whole set of functional tests is run as part of summary a "Repro String" is generated. This string can be used as a parameter to RunAllFunctionalTests and if so only indicated tests will be performed.

diff --git a/Engine/Source/Developer/FunctionalTesting/Classes/FunctionalAITest.h b/Engine/Source/Developer/FunctionalTesting/Classes/FunctionalAITest.h
--- a/Engine/Source/Developer/FunctionalTesting/Classes/FunctionalAITest.h
+++ b/Engine/Source/Developer/FunctionalTesting/Classes/FunctionalAITest.h
@@ -98,6 +98,7 @@ class AFunctionalAITest : public AFunctionalTest
 {
 	GENERATED_UCLASS_BODY()
 	
+protected:
 	UPROPERTY(EditAnywhere, BlueprintReadOnly, Category=AITest)
 	TArray<FAITestSpawnSet> SpawnSets;
 
@@ -117,7 +118,10 @@ class AFunctionalAITest : public AFunctionalTest
 	/** Called when a all AI finished spawning */
 	UPROPERTY(BlueprintAssignable)
 	FFunctionalTestEventSignature OnAllAISPawned;
-		
+
+	uint32 bSingleSetRun:1;
+
+public:
 	UFUNCTION(BlueprintCallable, Category = "Development")
 	virtual bool IsOneOfSpawnedPawns(AActor* Actor);
 
@@ -126,11 +130,14 @@ class AFunctionalAITest : public AFunctionalTest
 	virtual void Tick(float DeltaSeconds) override;
 	// AActor interface end
 
-	virtual bool StartTest() override;
+	virtual bool StartTest(const TArray<FString>& Params = TArray<FString>()) override;
 	virtual void FinishTest(TEnumAsByte<EFunctionalTestResult::Type> TestResult, const FString& Message) override;
 	virtual bool WantsToRunAgain() const override;
 	virtual void CleanUp() override;
 	virtual FString GetAdditionalTestFinishedMessage(EFunctionalTestResult::Type TestResult) const override;
+	virtual FString GetReproString() const override;
+
+	void AddSpawnedPawn(APawn& SpawnedPawn);
 
 protected:
 
diff --git a/Engine/Source/Developer/FunctionalTesting/Classes/FunctionalTest.h b/Engine/Source/Developer/FunctionalTesting/Classes/FunctionalTest.h
--- a/Engine/Source/Developer/FunctionalTesting/Classes/FunctionalTest.h
+++ b/Engine/Source/Developer/FunctionalTesting/Classes/FunctionalTest.h
@@ -70,8 +70,7 @@ class AFunctionalTest : public AActor
 
 public:
 
-	UFUNCTION(BlueprintCallable, Category="Development")
-	virtual bool StartTest();
+	virtual bool StartTest(const TArray<FString>& Params = TArray<FString>());
 
 	UFUNCTION(BlueprintCallable, Category="Development")
 	virtual void FinishTest(TEnumAsByte<EFunctionalTestResult::Type> TestResult, const FString& Message);
@@ -99,6 +98,8 @@ public:
 	 *	WantsToRunAgain calls). CleanUp gets called when all cycles are done. */
 	virtual void CleanUp();
 
+	virtual FString GetReproString() const { return GetFName().ToString(); }
+
 #if WITH_EDITOR
 	void PostEditChangeProperty(struct FPropertyChangedEvent& PropertyChangedEvent) override;
 #endif // WITH_EDITOR
diff --git a/Engine/Source/Developer/FunctionalTesting/Classes/FunctionalTestingManager.h b/Engine/Source/Developer/FunctionalTesting/Classes/FunctionalTestingManager.h
--- a/Engine/Source/Developer/FunctionalTesting/Classes/FunctionalTestingManager.h
+++ b/Engine/Source/Developer/FunctionalTesting/Classes/FunctionalTestingManager.h
@@ -5,6 +5,12 @@
 #include "Misc/AutomationTest.h"
 #include "FunctionalTestingManager.generated.h"
 
+namespace FFunctionalTesting
+{
+	extern const TCHAR* ReproStringTestSeparator;
+	extern const TCHAR* ReproStringParamsSeparator;
+}
+
 UCLASS(BlueprintType, MinimalAPI)
 class UFunctionalTestingManager : public UBlueprintFunctionLibrary
 {
@@ -22,7 +28,7 @@ class UFunctionalTestingManager : public UBlueprintFunctionLibrary
 	UFUNCTION(BlueprintCallable, Category="FunctionalTesting", meta=(HidePin="WorldContext", DefaultToSelf="WorldContext" ) )
 	/** Triggers in sequence all functional tests found on the level.
 	 *	@return true if any tests have been triggered */
-	static bool RunAllFunctionalTests(UObject* WorldContext, bool bNewLog=true, bool bRunLooped=false, bool bWaitForNavigationBuildFinish=true, bool bLoopOnlyFailedTests = true);
+	 static bool RunAllFunctionalTests(UObject* WorldContext, bool bNewLog = true, bool bRunLooped = false, bool bWaitForNavigationBuildFinish = true, FString FailedTestsReproString = TEXT(""));
 		
 	bool IsRunning() const { return bIsRunning; }
 	bool IsLooped() const { return bLooped; }
@@ -37,7 +43,7 @@ class UFunctionalTestingManager : public UBlueprintFunctionLibrary
 	static void AddError(const FText& InError);
 	static void AddWarning(const FText& InWarning);
 	static void AddLogItem(const FText& InLogItem);
-	
+
 private:
 	void LogMessage(const FString& MessageString, TSharedPtr<IMessageLogListing> LogListing = NULL);
 	
@@ -53,13 +59,21 @@ protected:
 	bool RunFirstValidTest();
 
 	void NotifyTestDone(class AFunctionalTest* FTest);
+
+	void SetReproString(FString ReproString);
+	void AllTestsDone();
 	
 	bool bIsRunning;
 	bool bLooped;
 	bool bWaitForNavigationBuildFinish;
 	bool bInitialDelayApplied;
-	bool bDiscardSuccessfulTests;
 	uint32 CurrentIteration;
 
+	FFunctionalTestDoneSignature TestFinishedObserver;
+	
+	FString GatheredFailedTestsReproString;
+	FString StartingReproString;
+	TArray<FString> TestReproStrings;
+
 	static FAutomationTestExecutionInfo* ExecutionInfo;
 };
diff --git a/Engine/Source/Developer/FunctionalTesting/Private/FunctionalAITest.cpp b/Engine/Source/Developer/FunctionalTesting/Private/FunctionalAITest.cpp
--- a/Engine/Source/Developer/FunctionalTesting/Private/FunctionalAITest.cpp
+++ b/Engine/Source/Developer/FunctionalTesting/Private/FunctionalAITest.cpp
@@ -9,16 +9,8 @@
 AFunctionalAITest::AFunctionalAITest( const class FPostConstructInitializeProperties& PCIP )
 	: Super(PCIP)
 	, CurrentSpawnSetIndex(INDEX_NONE)
+	, bSingleSetRun(false)
 {
-	/*struct FConstructorStatics
-	{
-		ConstructorHelpers::FObjectFinderOptional<UTexture2D> Texture;
-
-		FConstructorStatics()
-			: Texture(TEXT("/Engine/EditorResources/S_FTest"))
-		{
-		}
-	};*/
 }
 
 bool AFunctionalAITest::IsOneOfSpawnedPawns(AActor* Actor)
@@ -57,12 +49,22 @@ void AFunctionalAITest::BeginPlay()
 	Super::BeginPlay();
 }
 
-bool AFunctionalAITest::StartTest()
+bool AFunctionalAITest::StartTest(const TArray<FString>& Params)
 {
 	KillOffSpawnedPawns();
 	ClearPendingDelayedSpawns();
 
-	if (++CurrentSpawnSetIndex < SpawnSets.Num())
+	bSingleSetRun = Params.Num() > 0;
+	if (bSingleSetRun)
+	{
+		TTypeFromString<int32>::FromString(CurrentSpawnSetIndex, *Params[0]);
+	}
+	else
+	{
+		++CurrentSpawnSetIndex;
+	}
+
+	if (CurrentSpawnSetIndex < SpawnSets.Num())
 	{
 		UWorld* World = GetWorld();
 		check(World);
@@ -132,7 +134,7 @@ void AFunctionalAITest::FinishTest(TEnumAsByte<EFunctionalTestResult::Type> Test
 
 bool AFunctionalAITest::WantsToRunAgain() const
 {
-	return CurrentSpawnSetIndex + 1 < SpawnSets.Num();
+	return bSingleSetRun == false && CurrentSpawnSetIndex + 1 < SpawnSets.Num();
 }
 
 
@@ -153,7 +155,7 @@ FString AFunctionalAITest::GetAdditionalTestFinishedMessage(EFunctionalTestResul
 	{
 		if (CurrentSpawnSetName.Len() > 0 && CurrentSpawnSetName != TEXT("None"))
 		{
-			Result = FString::Printf(TEXT("set %s, pawns: "), *CurrentSpawnSetName);
+			Result = FString::Printf(TEXT("spawn set \'%s\', pawns: "), *CurrentSpawnSetName);
 		}
 		else
 		{
@@ -170,6 +172,13 @@ FString AFunctionalAITest::GetAdditionalTestFinishedMessage(EFunctionalTestResul
 	return Result;
 }
 
+FString AFunctionalAITest::GetReproString() const
+{
+	return FString::Printf(TEXT("%s%s%d"), *(GetFName().ToString())
+		, FFunctionalTesting::ReproStringParamsSeparator
+		, CurrentSpawnSetIndex);
+}
+
 void AFunctionalAITest::KillOffSpawnedPawns()
 {
 	for (int32 PawnIndex = 0; PawnIndex < SpawnedPawns.Num(); ++PawnIndex)
@@ -199,6 +208,12 @@ void AFunctionalAITest::Tick(float DeltaSeconds)
 	}
 }
 
+void AFunctionalAITest::AddSpawnedPawn(APawn& SpawnedPawn)
+{
+	SpawnedPawns.Add(&SpawnedPawn);
+	OnAISpawned.Broadcast(Cast<AAIController>(SpawnedPawn.GetController()), &SpawnedPawn);
+}
+
 //----------------------------------------------------------------------//
 // FAITestSpawnInfo
 //----------------------------------------------------------------------//
@@ -231,9 +246,7 @@ bool FAITestSpawnInfo::Spawn(AFunctionalAITest* AITest) const
 	}
 	else
 	{
-		AITest->SpawnedPawns.Add(SpawnedPawn);
-		AITest->OnAISpawned.Broadcast(Cast<AAIController>(SpawnedPawn->GetController()), SpawnedPawn);
-
+		AITest->AddSpawnedPawn(*SpawnedPawn);
 		bSuccessfullySpawned = true;
 	}
 
diff --git a/Engine/Source/Developer/FunctionalTesting/Private/FunctionalTest.cpp b/Engine/Source/Developer/FunctionalTesting/Private/FunctionalTest.cpp
--- a/Engine/Source/Developer/FunctionalTesting/Private/FunctionalTest.cpp
+++ b/Engine/Source/Developer/FunctionalTesting/Private/FunctionalTest.cpp
@@ -70,7 +70,7 @@ void AFunctionalTest::Tick(float DeltaSeconds)
 	}
 }
 
-bool AFunctionalTest::StartTest()
+bool AFunctionalTest::StartTest(const TArray<FString>& Params)
 {
 	FailureMessage = TEXT("");
 
@@ -99,6 +99,8 @@ void AFunctionalTest::FinishTest(TEnumAsByte<EFunctionalTestResult::Type> TestRe
 		return;
 	}
 
+	Result = TestResult;
+
 	bIsRunning = false;
 	SetActorTickEnabled(false);
 
@@ -116,56 +118,39 @@ void AFunctionalTest::FinishTest(TEnumAsByte<EFunctionalTestResult::Type> TestRe
 	}
 
 	const FText ResultText = FTestResultTypeEnum->GetEnumText( TestResult.GetValue() );
-	const FString OutMessage = FString::Printf(TEXT("%s> Result: %s> %s")
+	const FString OutMessage = FString::Printf(TEXT("%s %s: \"%s\'")
 		, *GetActorLabel()
 		, *ResultText.ToString()
 		, Message.IsEmpty() == false ? *Message : TEXT("Test finished") );
 	const FString AdditionalDetails = GetAdditionalTestFinishedMessage(TestResult) + FString::Printf(TEXT(", time %.2fs"), TotalTime);
 
 	AutoDestroyActors.Reset();
-
-	EMessageSeverity::Type MessageLogSeverity = EMessageSeverity::Info;
-	
+		
 	switch (TestResult.GetValue())
 	{
 		case EFunctionalTestResult::Invalid:
 		case EFunctionalTestResult::Error:
+		case EFunctionalTestResult::Failed:
 			UE_VLOG(this, LogFunctionalTest, Error, TEXT("%s"), *OutMessage);
-			MessageLogSeverity = EMessageSeverity::Error;
+			UFunctionalTestingManager::AddError(FText::FromString(OutMessage));
 			break;
+
 		case EFunctionalTestResult::Running:
 			UE_VLOG(this, LogFunctionalTest, Warning, TEXT("%s"), *OutMessage);
-			MessageLogSeverity = EMessageSeverity::Warning;
-			break;
-		case EFunctionalTestResult::Failed:
-			UE_VLOG(this, LogFunctionalTest, Error, TEXT("%s"), *OutMessage);
-			MessageLogSeverity = EMessageSeverity::Error;
+			UFunctionalTestingManager::AddWarning(FText::FromString(OutMessage));
 			break;
+		
 		default:
 			UE_VLOG(this, LogFunctionalTest, Log, TEXT("%s"), *OutMessage);
+			UFunctionalTestingManager::AddLogItem(FText::FromString(OutMessage));
 			break;
 	}
-
-	TSharedRef<FTokenizedMessage> LogMessage = FTokenizedMessage::Create(MessageLogSeverity, FText::FromString(GetActorLabel()))
-		->AddToken(FTextToken::Create(ResultText))
-		->AddToken(FTextToken::Create(FText::FromString(Message)))
-		->AddToken(FTextToken::Create(FText::FromString(AdditionalDetails)))
-		->AddToken(FTextToken::Create(FText::FromString(FailureMessage)));
-
-	switch (MessageLogSeverity)
+	
+	if (AdditionalDetails.IsEmpty() == false)
 	{
-	case EMessageSeverity::Error:
-		UFunctionalTestingManager::AddError(LogMessage->ToText());
-		break;
-	case EMessageSeverity::Warning:
-		UFunctionalTestingManager::AddWarning(LogMessage->ToText());
-		break;
-	default:
-		UFunctionalTestingManager::AddLogItem(LogMessage->ToText());
-		break;
+		UFunctionalTestingManager::AddLogItem(FText::FromString(AdditionalDetails));
 	}
 
-
 	TestFinishedObserver.ExecuteIfBound(this);
 }
 
diff --git a/Engine/Source/Developer/FunctionalTesting/Private/FunctionalTestingManager.cpp b/Engine/Source/Developer/FunctionalTesting/Private/FunctionalTestingManager.cpp
--- a/Engine/Source/Developer/FunctionalTesting/Private/FunctionalTestingManager.cpp
+++ b/Engine/Source/Developer/FunctionalTesting/Private/FunctionalTestingManager.cpp
@@ -4,6 +4,13 @@
 
 #define LOCTEXT_NAMESPACE "FunctionalTesting"
 
+namespace FFunctionalTesting
+{
+	const TCHAR* ReproStringTestSeparator = TEXT("@");
+	const TCHAR* ReproStringParamsSeparator = TEXT("#");
+}
+
+
 //struct FFuncTestingTickHelper : FTickableGameObject
 //{
 //	class UFunctionalTestingManager* Manager;
@@ -45,6 +52,7 @@ UFunctionalTestingManager::UFunctionalTestingManager( const class FPostConstruct
 	, bInitialDelayApplied(false)
 	, CurrentIteration(INDEX_NONE)
 {
+	TestFinishedObserver = FFunctionalTestDoneSignature::CreateUObject(this, &UFunctionalTestingManager::OnTestDone);
 }
 
 void UFunctionalTestingManager::SetUpTests()
@@ -52,7 +60,7 @@ void UFunctionalTestingManager::SetUpTests()
 	OnSetupTests.Broadcast();
 }
 
-bool UFunctionalTestingManager::RunAllFunctionalTests(UObject* WorldContext, bool bNewLog, bool bRunLooped, bool bInWaitForNavigationBuildFinish, bool bLoopOnlyFailedTests)
+bool UFunctionalTestingManager::RunAllFunctionalTests(UObject* WorldContext, bool bNewLog, bool bRunLooped, bool bInWaitForNavigationBuildFinish, FString ReproString)
 {
 	UFunctionalTestingManager* Manager = GetManager(WorldContext);
 
@@ -72,23 +80,34 @@ bool UFunctionalTestingManager::RunAllFunctionalTests(UObject* WorldContext, boo
 	Manager->bLooped = bRunLooped;
 	Manager->bWaitForNavigationBuildFinish = bInWaitForNavigationBuildFinish;
 	Manager->CurrentIteration = 0;
-	Manager->bDiscardSuccessfulTests = bLoopOnlyFailedTests;
+	Manager->TestsLeft.Reset();
+	Manager->SetReproString(ReproString);
 
-	for (TActorIterator<AFunctionalTest> It(GWorld); It; ++It)
+	Manager->SetUpTests();
+
+	if (Manager->TestReproStrings.Num() > 0)
 	{
-		AFunctionalTest* Test = (*It);
-		if (Test != NULL && Test->bIsEnabled == true)
-		{
-			Manager->AllTests.Add(Test);
-		}
+		static const FText RunningFromReproString = LOCTEXT("RunningFromReproString", "Running tests indicated by Repro String:");
+		AddLogItem(RunningFromReproString);
+		AddLogItem(FText::FromString(ReproString));
+		Manager->TriggerFirstValidTest();
 	}
-
-	if (Manager->AllTests.Num() > 0)
+	else
 	{
-		Manager->SetUpTests();
-		Manager->TestsLeft = Manager->AllTests;
+		for (TActorIterator<AFunctionalTest> It(GWorld); It; ++It)
+		{
+			AFunctionalTest* Test = (*It);
+			if (Test != NULL && Test->bIsEnabled == true)
+			{
+				Manager->AllTests.Add(Test);
+			}
+		}
 
-		Manager->TriggerFirstValidTest();
+		if (Manager->AllTests.Num() > 0)
+		{
+			Manager->TestsLeft = Manager->AllTests;
+			Manager->TriggerFirstValidTest();
+		}
 	}
 
 	if (Manager->bIsRunning == false)
@@ -118,6 +137,10 @@ void UFunctionalTestingManager::TriggerFirstValidTest()
 	if (bInitialDelayApplied == true && (bWaitForNavigationBuildFinish == false || UNavigationSystem::IsNavigationBeingBuilt(World) == false))
 	{
 		bIsRunning = RunFirstValidTest();
+		if (bIsRunning == false)
+		{
+			AllTestsDone();
+		}
 	}
 	else
 	{
@@ -154,6 +177,15 @@ void UFunctionalTestingManager::OnTestDone(AFunctionalTest* FTest)
 
 void UFunctionalTestingManager::NotifyTestDone(AFunctionalTest* FTest)
 {
+	if (FTest->IsSuccessful() == false)
+	{
+		if (GatheredFailedTestsReproString.IsEmpty() == false)
+		{
+			GatheredFailedTestsReproString += FFunctionalTesting::ReproStringTestSeparator;
+		}
+		GatheredFailedTestsReproString += FTest->GetReproString();
+	}
+
 	if (FTest->WantsToRunAgain() == false)
 	{
 		TestsLeft.RemoveSingle(FTest);
@@ -164,7 +196,7 @@ void UFunctionalTestingManager::NotifyTestDone(AFunctionalTest* FTest)
 		FTest->CleanUp();
 	}
 
-	if (TestsLeft.Num() > 0)
+	if (TestsLeft.Num() > 0 || TestReproStrings.Num() > 0)
 	{
 		bIsRunning = RunFirstValidTest();
 	}
@@ -175,55 +207,127 @@ void UFunctionalTestingManager::NotifyTestDone(AFunctionalTest* FTest)
 
 	if (bIsRunning == false)
 	{
-		if (bLooped == true)
-		{
-			++CurrentIteration;
-			TestsLeft = AllTests;
-			FFormatNamedArguments Arguments;
-			Arguments.Add(TEXT("IterationIndex"), CurrentIteration);
-			AddLogItem(FText::Format(LOCTEXT("StartingIteration", "----- Starting iteration {IterationIndex} -----"), Arguments));
-			bIsRunning = RunFirstValidTest();
-		}
-		else
+		AllTestsDone();
+	}
+}
+
+void UFunctionalTestingManager::AllTestsDone()
+{
+	if (GatheredFailedTestsReproString.IsEmpty() == false)
+	{
+		static const FText ReproStringLabel = LOCTEXT("ReproString", "Repro String:");
+		AddLogItem(ReproStringLabel);
+		AddLogItem(FText::FromString(GatheredFailedTestsReproString));
+	}
+
+	if (bLooped == true)
+	{
+		++CurrentIteration;
+
+		// reset
+		ensure(TestReproStrings.Num() == 0);
+		SetReproString(StartingReproString);
+		GatheredFailedTestsReproString = TEXT("");
+		TestsLeft = AllTests;
+
+		FFormatNamedArguments Arguments;
+		Arguments.Add(TEXT("IterationIndex"), CurrentIteration);
+		AddLogItem(FText::Format(LOCTEXT("StartingIteration", "----- Starting iteration {IterationIndex} -----"), Arguments));
+		bIsRunning = RunFirstValidTest();
+		if (bIsRunning == false)
 		{
-			AddLogItem(LOCTEXT("TestDone", "DONE."));
-			RemoveFromRoot();
+			static const FText FailedToStartAnotherIterationText = LOCTEXT("FailedToStartAnotherIteration", "Failed to start another iteration.");
+			AddWarning(FailedToStartAnotherIterationText);
 		}
 	}
+	else
+	{
+		AddLogItem(LOCTEXT("TestDone", "DONE."));
+		RemoveFromRoot();
+	}
 }
 
 bool UFunctionalTestingManager::RunFirstValidTest()
 {
-	for (int32 Index = 0; Index < TestsLeft.Num(); ++Index)
+	bool bTestSuccessfullyTriggered = false;
+
+	if (TestReproStrings.Num() > 0)
 	{
-		bool bRemove = TestsLeft[Index] == NULL;
-		if (TestsLeft[Index] != NULL)
+		UWorld* World = GEngine->GetWorldFromContextObject(GetOuter());
+		UObject* TestsOuter = World ? (UObject*)(World->PersistentLevel) : (UObject*)(ANY_PACKAGE);
+
+		while (TestReproStrings.Num() > 0)
 		{
-			ensure(TestsLeft[Index]->bIsEnabled);
-			TestsLeft[Index]->TestFinishedObserver = FFunctionalTestDoneSignature::CreateUObject(this, &UFunctionalTestingManager::OnTestDone);
-			if (TestsLeft[Index]->StartTest())
+			TArray<FString> TestParams;
+			const FString SingleTestReproString = TestReproStrings.Pop(/*bAllowShrinking=*/false);
+			SingleTestReproString.ParseIntoArray(&TestParams, TEXT("#"), /*InCullEmpty=*/true);
+			
+			// first param is the test name. Look for it			
+			AFunctionalTest* TestToRun = FindObject<AFunctionalTest>(TestsOuter, *TestParams[0]);
+			TestParams.RemoveAt(0, 1, /*bAllowShrinking=*/false);
+			if (TestToRun)
 			{
-				return true;
+				TestToRun->TestFinishedObserver = TestFinishedObserver;
+				if (TestToRun->StartTest(TestParams))
+				{
+					bTestSuccessfullyTriggered = true;
+					break;
+				}
+				else
+				{
+					AddWarning(FText::FromString(FString::Printf(TEXT("Test \'%s\' failed to start"), *TestToRun->GetName())));
+				}
 			}
 			else
 			{
-				bRemove = true;
+				AddWarning(FText::FromString(FString::Printf(TEXT("Unable to find test \'%s\'"), *TestToRun->GetName())));
 			}
 		}
-		
-		if (bRemove)
+	}
+	else
+	{
+		for (int32 Index = TestsLeft.Num()-1; Index >= 0; --Index)
 		{
-			TestsLeft.RemoveAtSwap(Index, 1, false);
+			bool bRemove = TestsLeft[Index] == NULL;
+			if (TestsLeft[Index] != NULL)
+			{
+				ensure(TestsLeft[Index]->bIsEnabled);
+				TestsLeft[Index]->TestFinishedObserver = TestFinishedObserver;
+				if (TestsLeft[Index]->StartTest())
+				{
+					bTestSuccessfullyTriggered = true;
+					break;
+				}
+				else
+				{
+					AddWarning(FText::FromString(FString::Printf(TEXT("Test: %s failed to start"), *TestsLeft[Index]->GetName())));
+					bRemove = true;
+				}
+			}
+
+			if (bRemove)
+			{
+				TestsLeft.RemoveAtSwap(Index, 1, false);
+			}
 		}
 	}
 
-	return false;
+	return bTestSuccessfullyTriggered;
 }
 
 void UFunctionalTestingManager::TickMe(float DeltaTime)
 {
+		
+}
 
-	
+void UFunctionalTestingManager::SetReproString(FString ReproString)
+{
+	TestReproStrings.Reset();
+	StartingReproString = ReproString;
+	if (ReproString.IsEmpty() == false)
+	{
+		ReproString.ParseIntoArray(&TestReproStrings, FFunctionalTesting::ReproStringTestSeparator, /*InCullEmpty=*/true);
+	}
 }
 
 //----------------------------------------------------------------------//
