Fixed FunctionalTestingManager not being removed from the rootset on world cleanup #UE4

Also FunctionalFramework's pch no longer including Engine.h

diff --git a/Engine/Source/Developer/FunctionalTesting/Classes/FuncTestRenderingComponent.h b/Engine/Source/Developer/FunctionalTesting/Classes/FuncTestRenderingComponent.h
--- a/Engine/Source/Developer/FunctionalTesting/Classes/FuncTestRenderingComponent.h
+++ b/Engine/Source/Developer/FunctionalTesting/Classes/FuncTestRenderingComponent.h
@@ -1,7 +1,7 @@
 // Copyright 1998-2014 Epic Games, Inc. All Rights Reserved.
 
 #pragma once
-
+#include "Components/PrimitiveComponent.h"
 #include "FuncTestRenderingComponent.generated.h"
 
 UCLASS(hidecategories=Object, editinlinenew)
diff --git a/Engine/Source/Developer/FunctionalTesting/Classes/FunctionalAITest.h b/Engine/Source/Developer/FunctionalTesting/Classes/FunctionalAITest.h
--- a/Engine/Source/Developer/FunctionalTesting/Classes/FunctionalAITest.h
+++ b/Engine/Source/Developer/FunctionalTesting/Classes/FunctionalAITest.h
@@ -3,9 +3,9 @@
 #pragma once
 
 #include "AITypes.h"
+#include "GameFramework/Pawn.h"
 #include "FunctionalAITest.generated.h"
 
-class UWorld;
 class AFunctionalAITest;
 
 DECLARE_DYNAMIC_MULTICAST_DELEGATE_TwoParams(FFunctionalTestAISpawned, AAIController*, Controller, APawn*, Pawn);
diff --git a/Engine/Source/Developer/FunctionalTesting/Classes/FunctionalTest.h b/Engine/Source/Developer/FunctionalTesting/Classes/FunctionalTest.h
--- a/Engine/Source/Developer/FunctionalTesting/Classes/FunctionalTest.h
+++ b/Engine/Source/Developer/FunctionalTesting/Classes/FunctionalTest.h
@@ -1,9 +1,11 @@
 // Copyright 1998-2014 Epic Games, Inc. All Rights Reserved.
 
 #pragma once
-
+#include "GameFramework/Actor.h"
 #include "FunctionalTest.generated.h"
 
+class UBillboardComponent;
+
 UENUM()
 namespace EFunctionalTestResult
 {
@@ -106,6 +108,7 @@ public:
 
 	// AActor interface begin
 	virtual void Tick(float DeltaSeconds) override;
+	virtual void EndPlay(const EEndPlayReason::Type EndPlayReason) override;
 	// AActor interface end
 
 	bool IsSuccessful() const { return Result == EFunctionalTestResult::Succeeded; }
diff --git a/Engine/Source/Developer/FunctionalTesting/Classes/FunctionalTestingManager.h b/Engine/Source/Developer/FunctionalTesting/Classes/FunctionalTestingManager.h
--- a/Engine/Source/Developer/FunctionalTesting/Classes/FunctionalTestingManager.h
+++ b/Engine/Source/Developer/FunctionalTesting/Classes/FunctionalTestingManager.h
@@ -3,8 +3,11 @@
 #pragma once
 
 #include "Misc/AutomationTest.h"
+#include "Kismet/BlueprintFunctionLibrary.h"
 #include "FunctionalTestingManager.generated.h"
 
+class UWorld;
+
 namespace FFunctionalTesting
 {
 	extern const TCHAR* ReproStringTestSeparator;
@@ -62,6 +65,9 @@ protected:
 
 	void SetReproString(FString ReproString);
 	void AllTestsDone();
+
+	void OnWorldCleanedUp(UWorld* World, bool bSessionEnded, bool bCleanupResources);
+	virtual UWorld* GetWorld() const;
 	
 	bool bIsRunning;
 	bool bLooped;
diff --git a/Engine/Source/Developer/FunctionalTesting/Private/FuncTestManager.cpp b/Engine/Source/Developer/FunctionalTesting/Private/FuncTestManager.cpp
--- a/Engine/Source/Developer/FunctionalTesting/Private/FuncTestManager.cpp
+++ b/Engine/Source/Developer/FunctionalTesting/Private/FuncTestManager.cpp
@@ -1,6 +1,7 @@
 // Copyright 1998-2014 Epic Games, Inc. All Rights Reserved.
 
 #include "FunctionalTestingPrivatePCH.h"
+#include "Engine/World.h"
 
 DEFINE_LOG_CATEGORY(LogFunctionalTest);
 
diff --git a/Engine/Source/Developer/FunctionalTesting/Private/FuncTestRenderingComponent.cpp b/Engine/Source/Developer/FunctionalTesting/Private/FuncTestRenderingComponent.cpp
--- a/Engine/Source/Developer/FunctionalTesting/Private/FuncTestRenderingComponent.cpp
+++ b/Engine/Source/Developer/FunctionalTesting/Private/FuncTestRenderingComponent.cpp
@@ -1,6 +1,8 @@
 // Copyright 1998-2014 Epic Games, Inc. All Rights Reserved.
 
 #include "FunctionalTestingPrivatePCH.h"
+#include "PrimitiveViewRelevance.h"
+#include "PrimitiveSceneProxy.h"
 
 //----------------------------------------------------------------------//
 // FFTestRenderingSceneProxy
diff --git a/Engine/Source/Developer/FunctionalTesting/Private/FunctionalTest.cpp b/Engine/Source/Developer/FunctionalTesting/Private/FunctionalTest.cpp
--- a/Engine/Source/Developer/FunctionalTesting/Private/FunctionalTest.cpp
+++ b/Engine/Source/Developer/FunctionalTesting/Private/FunctionalTest.cpp
@@ -154,6 +154,13 @@ void AFunctionalTest::FinishTest(TEnumAsByte<EFunctionalTestResult::Type> TestRe
 	TestFinishedObserver.ExecuteIfBound(this);
 }
 
+void AFunctionalTest::EndPlay(const EEndPlayReason::Type EndPlayReason)
+{
+	TestFinishedObserver.Unbind();
+
+	Super::EndPlay(EndPlayReason);
+}
+
 void AFunctionalTest::CleanUp()
 {
 	FailureMessage = TEXT("");
diff --git a/Engine/Source/Developer/FunctionalTesting/Private/FunctionalTestingManager.cpp b/Engine/Source/Developer/FunctionalTesting/Private/FunctionalTestingManager.cpp
--- a/Engine/Source/Developer/FunctionalTesting/Private/FunctionalTestingManager.cpp
+++ b/Engine/Source/Developer/FunctionalTesting/Private/FunctionalTestingManager.cpp
@@ -1,6 +1,9 @@
 // Copyright 1998-2014 Epic Games, Inc. All Rights Reserved.
 
 #include "FunctionalTestingPrivatePCH.h"
+#if WITH_EDITOR
+#include "Editor.h"
+#endif
 
 #define LOCTEXT_NAMESPACE "FunctionalTesting"
 
@@ -52,7 +55,10 @@ UFunctionalTestingManager::UFunctionalTestingManager( const class FPostConstruct
 	, bInitialDelayApplied(false)
 	, CurrentIteration(INDEX_NONE)
 {
-	TestFinishedObserver = FFunctionalTestDoneSignature::CreateUObject(this, &UFunctionalTestingManager::OnTestDone);
+	if (HasAnyFlags(RF_ClassDefaultObject) == false)
+	{
+		TestFinishedObserver = FFunctionalTestDoneSignature::CreateUObject(this, &UFunctionalTestingManager::OnTestDone);
+	}
 }
 
 void UFunctionalTestingManager::SetUpTests()
@@ -119,8 +125,7 @@ bool UFunctionalTestingManager::RunAllFunctionalTests(UObject* WorldContext, boo
 #if WITH_EDITOR
 	FEditorDelegates::EndPIE.AddUObject(Manager, &UFunctionalTestingManager::OnEndPIE);
 #endif // WITH_EDITOR
-	Manager->AddToRoot();
-
+	
 	return true;
 }
 
@@ -131,7 +136,7 @@ void UFunctionalTestingManager::OnEndPIE(const bool bIsSimulating)
 
 void UFunctionalTestingManager::TriggerFirstValidTest()
 {
-	UWorld* World = GEngine->GetWorldFromContextObject(GetOuter()); 
+	UWorld* World = GetWorld();
 	bIsRunning = World != NULL && World->GetNavigationSystem() != NULL;
 
 	if (bInitialDelayApplied == true && (bWaitForNavigationBuildFinish == false || UNavigationSystem::IsNavigationBeingBuilt(World) == false))
@@ -159,11 +164,29 @@ UFunctionalTestingManager* UFunctionalTestingManager::GetManager(UObject* WorldC
 		UObject* Outer = WorldContext ? WorldContext : (UObject*)GetTransientPackage();
 		Manager = NewObject<UFunctionalTestingManager>(Outer);
 		FFunctionalTestingModule::Get()->SetScript(Manager);
+
+		// add to root and get notified on world cleanup to remove from root on map cleanup
+		Manager->AddToRoot();
+		FWorldDelegates::OnWorldCleanup.AddUObject(Manager, &UFunctionalTestingManager::OnWorldCleanedUp);
 	}
 
 	return Manager;
 }
 
+UWorld* UFunctionalTestingManager::GetWorld() const
+{
+	return GEngine->GetWorldFromContextObject(GetOuter());
+}
+
+void UFunctionalTestingManager::OnWorldCleanedUp(UWorld* World, bool bSessionEnded, bool bCleanupResources)
+{
+	UWorld* MyWorld = GetWorld();
+	if (MyWorld == World)
+	{
+		RemoveFromRoot();
+	}
+}
+
 void UFunctionalTestingManager::OnTestDone(AFunctionalTest* FTest)
 {
 	// add a delay
@@ -253,7 +276,7 @@ bool UFunctionalTestingManager::RunFirstValidTest()
 
 	if (TestReproStrings.Num() > 0)
 	{
-		UWorld* World = GEngine->GetWorldFromContextObject(GetOuter());
+		UWorld* World = GetWorld();
 		UObject* TestsOuter = World ? (UObject*)(World->PersistentLevel) : (UObject*)(ANY_PACKAGE);
 
 		while (TestReproStrings.Num() > 0)
diff --git a/Engine/Source/Developer/FunctionalTesting/Private/FunctionalTestingPrivatePCH.h b/Engine/Source/Developer/FunctionalTesting/Private/FunctionalTestingPrivatePCH.h
--- a/Engine/Source/Developer/FunctionalTesting/Private/FunctionalTestingPrivatePCH.h
+++ b/Engine/Source/Developer/FunctionalTesting/Private/FunctionalTestingPrivatePCH.h
@@ -1,12 +1,18 @@
 // Copyright 1998-2014 Epic Games, Inc. All Rights Reserved.
 
-#ifndef __FunctionalTestingPrivatePCH_h__
-#define __FunctionalTestingPrivatePCH_h__
+#pragma once
 
 #include "Engine.h"
-#if WITH_EDITOR
-#include "UnrealEd.h"
-#endif
+
+#include "Core.h"
+#include "CoreUObject.h"
+#include "InputCore.h"
+#include "EngineDefines.h"
+#include "EngineSettings.h"
+#include "EngineStats.h"
+#include "EngineLogs.h"
+#include "EngineGlobals.h"
+
 #include "FunctionalTest.h"
 #include "FunctionalAITest.h"
 #include "Developer/MessageLog/Public/MessageLogModule.h"
@@ -16,5 +22,3 @@
 #include "FunctionalTestingManager.h"
 #include "FunctionalTestingModule.h"
 #include "FuncTestRenderingComponent.h"
-
-#endif // __FunctionalTestingPrivatePCH_h__
diff --git a/Engine/Source/Developer/FunctionalTesting/Private/Tests/FuncTestPerforming.cpp b/Engine/Source/Developer/FunctionalTesting/Private/Tests/FuncTestPerforming.cpp
--- a/Engine/Source/Developer/FunctionalTesting/Private/Tests/FuncTestPerforming.cpp
+++ b/Engine/Source/Developer/FunctionalTesting/Private/Tests/FuncTestPerforming.cpp
@@ -2,8 +2,13 @@
 
 #include "FunctionalTestingPrivatePCH.h"
 #if WITH_EDITOR
+#include "Engine/Brush.h"
+#include "Editor/EditorEngine.h"
+#include "EngineUtils.h"
 #include "Editor/UnrealEd/Classes/Editor/EditorEngine.h"
 #include "Editor/UnrealEd/Public/FileHelpers.h"
+
+extern UNREALED_API class UEditorEngine* GEditor;
 #endif // WITH_EDITOR
 
 #define LOCTEXT_NAMESPACE "FunctionalTesting"
