Fixed Functional Testing crashing in game-mode #UE4

diff --git a/Engine/Source/Developer/FunctionalTesting/Private/Tests/FuncTestPerforming.cpp b/Engine/Source/Developer/FunctionalTesting/Private/Tests/FuncTestPerforming.cpp
--- a/Engine/Source/Developer/FunctionalTesting/Private/Tests/FuncTestPerforming.cpp
+++ b/Engine/Source/Developer/FunctionalTesting/Private/Tests/FuncTestPerforming.cpp
@@ -189,7 +189,10 @@ bool FFunctionalTestingMapsBase::RunTest(const FString& Parameters)
 	else
 #endif // WITH_EDITOR
 	{
-		GEngine->Exec( NULL, *FString::Printf(TEXT("Open %s"), *MapName));
+		check(GEngine->GetWorldContexts().Num() == 1);
+		check(GEngine->GetWorldContexts()[0].WorldType == EWorldType::Game);
+
+		GEngine->Exec(GEngine->GetWorldContexts()[0].World(), *FString::Printf(TEXT("Open %s"), *MapName));
 	}
 	ADD_LATENT_AUTOMATION_COMMAND(FStartFTestsOnMap());
 
