Fixed FunctionalTesting adding a new tab to MessageLog when just being asked for tests list #UE4

diff --git a/Engine/Source/Developer/FunctionalTesting/Private/Tests/FuncTestPerforming.cpp b/Engine/Source/Developer/FunctionalTesting/Private/Tests/FuncTestPerforming.cpp
--- a/Engine/Source/Developer/FunctionalTesting/Private/Tests/FuncTestPerforming.cpp
+++ b/Engine/Source/Developer/FunctionalTesting/Private/Tests/FuncTestPerforming.cpp
@@ -119,6 +119,16 @@ public:
 namespace
 {
 	FFunctionalTestingMaps FFunctionalTestingMapsAutomationTestInstance(TEXT("FFunctionalTestingMaps"));
+
+	void CreateFTestingLogPage()
+	{
+		static bool bMessageLogPageCreated = false;
+		if (bMessageLogPageCreated == false)
+		{
+			FMessageLog("FunctionalTestingLog").NewPage(LOCTEXT("NewLogLabel", "Functional Test"));
+			bMessageLogPageCreated = true;
+		}
+	}
 }
 
 /** 
@@ -160,8 +170,6 @@ void FFunctionalTestingMapsBase::GetTests(TArray<FString>& OutBeautifiedNames, T
 			}
 		}
 	}
-
-	FMessageLog("FunctionalTestingLog").NewPage(LOCTEXT("NewLogLabel", "Functional Test"));
 }
 
 /** 
@@ -174,6 +182,8 @@ bool FFunctionalTestingMapsBase::RunTest(const FString& Parameters)
 {
 	FString MapName = Parameters;
 
+	CreateFTestingLogPage();
+
 	SetSuppressLogs(true);
 	ExecutionInfo.Clear();
 	UFunctionalTestingManager::SetAutomationExecutionInfo(&ExecutionInfo);
