Fixed FunctionalTesting framework not being able detect if the testing has already finished #UE4

diff --git a/Engine/Source/Developer/FunctionalTesting/Classes/FunctionalTestingManager.h b/Engine/Source/Developer/FunctionalTesting/Classes/FunctionalTestingManager.h
--- a/Engine/Source/Developer/FunctionalTesting/Classes/FunctionalTestingManager.h
+++ b/Engine/Source/Developer/FunctionalTesting/Classes/FunctionalTestingManager.h
@@ -34,6 +34,7 @@ class UFunctionalTestingManager : public UBlueprintFunctionLibrary
 	 static bool RunAllFunctionalTests(UObject* WorldContext, bool bNewLog = true, bool bRunLooped = false, bool bWaitForNavigationBuildFinish = true, FString FailedTestsReproString = TEXT(""));
 		
 	bool IsRunning() const { return bIsRunning; }
+	bool IsFinished() const { return bFinished; }
 	bool IsLooped() const { return bLooped; }
 	void SetLooped(const bool bNewLooped) { bLooped = bNewLooped; }
 
@@ -70,6 +71,7 @@ protected:
 	virtual UWorld* GetWorld() const;
 	
 	bool bIsRunning;
+	bool bFinished;
 	bool bLooped;
 	bool bWaitForNavigationBuildFinish;
 	bool bInitialDelayApplied;
diff --git a/Engine/Source/Developer/FunctionalTesting/Private/FuncTestManager.cpp b/Engine/Source/Developer/FunctionalTesting/Private/FuncTestManager.cpp
--- a/Engine/Source/Developer/FunctionalTesting/Private/FuncTestManager.cpp
+++ b/Engine/Source/Developer/FunctionalTesting/Private/FuncTestManager.cpp
@@ -15,6 +15,11 @@ bool FFuncTestManager::IsRunning() const
 	return TestScript.IsValid() && TestScript->IsRunning();
 }
 
+bool FFuncTestManager::IsFinished() const
+{
+	return TestScript.IsValid() && TestScript->IsFinished();
+}
+
 void FFuncTestManager::SetLooping(const bool bLoop)
 { 
 	if (TestScript.IsValid())
diff --git a/Engine/Source/Developer/FunctionalTesting/Private/FuncTestManager.h b/Engine/Source/Developer/FunctionalTesting/Private/FuncTestManager.h
--- a/Engine/Source/Developer/FunctionalTesting/Private/FuncTestManager.h
+++ b/Engine/Source/Developer/FunctionalTesting/Private/FuncTestManager.h
@@ -14,6 +14,8 @@ public:
 	virtual void RunAllTestsOnMap(bool bClearLog, bool bRunLooped) override;
 	
 	virtual bool IsRunning() const override;
+
+	virtual bool IsFinished() const override;
 	
 	virtual void SetScript(class UFunctionalTestingManager* NewScript) override;
 
diff --git a/Engine/Source/Developer/FunctionalTesting/Private/FunctionalTestingManager.cpp b/Engine/Source/Developer/FunctionalTesting/Private/FunctionalTestingManager.cpp
--- a/Engine/Source/Developer/FunctionalTesting/Private/FunctionalTestingManager.cpp
+++ b/Engine/Source/Developer/FunctionalTesting/Private/FunctionalTestingManager.cpp
@@ -58,6 +58,7 @@ FAutomationTestExecutionInfo* UFunctionalTestingManager::ExecutionInfo = NULL;
 UFunctionalTestingManager::UFunctionalTestingManager( const FObjectInitializer& ObjectInitializer )
 	: Super(ObjectInitializer)
 	, bIsRunning(false)
+	, bFinished(false)
 	, bLooped(false)
 	, bWaitForNavigationBuildFinish(false)
 	, bInitialDelayApplied(false)
@@ -91,6 +92,7 @@ bool UFunctionalTestingManager::RunAllFunctionalTests(UObject* WorldContext, boo
 		FunctionalTestingLog.NewPage(LOCTEXT("NewLogLabel", "Functional Test"));
 	}
 
+	Manager->bFinished = false;
 	Manager->bLooped = bRunLooped;
 	Manager->bWaitForNavigationBuildFinish = bInWaitForNavigationBuildFinish;
 	Manager->CurrentIteration = 0;
@@ -274,6 +276,7 @@ void UFunctionalTestingManager::AllTestsDone()
 	}
 	else
 	{
+		bFinished = true;
 		AddLogItem(LOCTEXT("TestDone", "DONE."));
 		RemoveFromRoot();
 	}
diff --git a/Engine/Source/Developer/FunctionalTesting/Private/Tests/FuncTestPerforming.cpp b/Engine/Source/Developer/FunctionalTesting/Private/Tests/FuncTestPerforming.cpp
--- a/Engine/Source/Developer/FunctionalTesting/Private/Tests/FuncTestPerforming.cpp
+++ b/Engine/Source/Developer/FunctionalTesting/Private/Tests/FuncTestPerforming.cpp
@@ -62,14 +62,17 @@ DEFINE_LATENT_AUTOMATION_COMMAND(FTriggerFTests);
 bool FTriggerFTests::Update()
 {
 	IFuncTestManager* Manager = FFunctionalTestingModule::Get();
-	// if tests have been already triggered by level script just make sure it's not looping
-	if (Manager->IsRunning())
+	if (Manager->IsFinished() == false)
 	{
-		FFunctionalTestingModule::Get()->SetLooping(false);
-	}
-	else
-	{
-		FFunctionalTestingModule::Get()->RunAllTestsOnMap(false, false);
+		// if tests have been already triggered by level script just make sure it's not looping
+		if (Manager->IsRunning())
+		{
+			FFunctionalTestingModule::Get()->SetLooping(false);
+		}
+		else
+		{
+			FFunctionalTestingModule::Get()->RunAllTestsOnMap(false, false);
+		}
 	}
 
 	return true;
diff --git a/Engine/Source/Developer/FunctionalTesting/Public/IFuncTestManager.h b/Engine/Source/Developer/FunctionalTesting/Public/IFuncTestManager.h
--- a/Engine/Source/Developer/FunctionalTesting/Public/IFuncTestManager.h
+++ b/Engine/Source/Developer/FunctionalTesting/Public/IFuncTestManager.h
@@ -9,6 +9,8 @@ public:
 	virtual void RunAllTestsOnMap(bool bClearLog, bool bRunLooped) = 0;
 
 	virtual bool IsRunning() const = 0;
+
+	virtual bool IsFinished() const = 0;
 	
 	virtual void SetScript(class UFunctionalTestingManager* NewScript) = 0;
 
