Improvements to functional testing framework #UE4

- improvements to debug drawing - a test can declare which level actors it's using and point at them while being selected in the editor
- added a "comment" field in propertied for a quick summary

diff --git a/Engine/Source/Developer/FunctionalTesting/Classes/FuncTestRenderingComponent.h b/Engine/Source/Developer/FunctionalTesting/Classes/FuncTestRenderingComponent.h
--- a/Engine/Source/Developer/FunctionalTesting/Classes/FuncTestRenderingComponent.h
+++ b/Engine/Source/Developer/FunctionalTesting/Classes/FuncTestRenderingComponent.h
@@ -11,6 +11,7 @@ class UFuncTestRenderingComponent : public UPrimitiveComponent
 
 	// Begin UPrimitiveComponent Interface
 	virtual FPrimitiveSceneProxy* CreateSceneProxy() override;
+	virtual bool ShouldRecreateProxyOnUpdateTransform() const override { return true; }
 	// End UPrimitiveComponent Interface
 
 	// Begin USceneComponent Interface
diff --git a/Engine/Source/Developer/FunctionalTesting/Classes/FunctionalAITest.h b/Engine/Source/Developer/FunctionalTesting/Classes/FunctionalAITest.h
--- a/Engine/Source/Developer/FunctionalTesting/Classes/FunctionalAITest.h
+++ b/Engine/Source/Developer/FunctionalTesting/Classes/FunctionalAITest.h
@@ -136,6 +136,7 @@ public:
 
 	virtual bool StartTest(const TArray<FString>& Params = TArray<FString>()) override;
 	virtual bool WantsToRunAgain() const override;
+	virtual void GatherRelevantActors(TArray<AActor*>& OutActors) const override;
 	virtual void CleanUp() override;
 	virtual FString GetAdditionalTestFinishedMessage(EFunctionalTestResult::Type TestResult) const override;
 	virtual FString GetReproString() const override;
diff --git a/Engine/Source/Developer/FunctionalTesting/Classes/FunctionalTest.h b/Engine/Source/Developer/FunctionalTesting/Classes/FunctionalTest.h
--- a/Engine/Source/Developer/FunctionalTesting/Classes/FunctionalTest.h
+++ b/Engine/Source/Developer/FunctionalTesting/Classes/FunctionalTest.h
@@ -68,6 +68,9 @@ public:
 
 	UPROPERTY(EditAnywhere, BlueprintReadWrite, Category = FunctionalTesting)
 	FRandomStream RandomNumbersStream;
+
+	UPROPERTY(EditAnywhere, BlueprintReadOnly, Category = FunctionalTesting)
+	FString Description;
 	
 #if WITH_EDITORONLY_DATA
 	UPROPERTY()
@@ -90,6 +93,12 @@ public:
 	UFUNCTION(BlueprintCallable, Category="Development")
 	virtual void SetTimeLimit(float NewTimeLimit, TEnumAsByte<EFunctionalTestResult::Type> ResultWhenTimeRunsOut);
 
+	/** Used by debug drawing to gather actors this test is using and point at them on the level to better understand test's setup */
+	UFUNCTION(BlueprintImplementableEvent, Category = "FunctionalTesting")
+	virtual TArray<AActor*> DebugGatherRelevantActors() const;
+
+	virtual void GatherRelevantActors(TArray<AActor*>& OutActors) const;
+
 	/** retrieves information whether test wants to have another run just after finishing */
 	UFUNCTION(BlueprintImplementableEvent, Category="FunctionalTesting")
 	virtual bool WantsToRunAgain() const;
@@ -111,6 +120,8 @@ public:
 
 #if WITH_EDITOR
 	void PostEditChangeProperty(struct FPropertyChangedEvent& PropertyChangedEvent) override;
+
+	static void OnSelectObject(UObject* NewSelection);
 #endif // WITH_EDITOR
 
 	// AActor interface begin
diff --git a/Engine/Source/Developer/FunctionalTesting/Private/FuncTestRenderingComponent.cpp b/Engine/Source/Developer/FunctionalTesting/Private/FuncTestRenderingComponent.cpp
--- a/Engine/Source/Developer/FunctionalTesting/Private/FuncTestRenderingComponent.cpp
+++ b/Engine/Source/Developer/FunctionalTesting/Private/FuncTestRenderingComponent.cpp
@@ -11,9 +11,24 @@ class FFTestRenderingSceneProxy : public FPrimitiveSceneProxy
 {
 public:
 	/** Initialization constructor. */
-	FFTestRenderingSceneProxy(const UFuncTestRenderingComponent* InComponent)
-		: FPrimitiveSceneProxy(InComponent)
+	FFTestRenderingSceneProxy(const UFuncTestRenderingComponent& InComponent)
+		: FPrimitiveSceneProxy(&InComponent)
 	{
+		OwningTest = Cast<AFunctionalTest>(InComponent.GetOwner());
+		if (OwningTest)
+		{
+			TArray<AActor*> RelevantActors;
+			OwningTest->GatherRelevantActors(RelevantActors);
+			
+			TestActorLocation = OwningTest->GetActorLocation();
+			for (auto Actor : RelevantActors)
+			{
+				if (Actor)
+				{
+					Locations.Add(Actor->GetActorLocation());
+				}
+			}
+		}
 	}
 
 	virtual FPrimitiveViewRelevance GetViewRelevance(const FSceneView* View)
@@ -25,11 +40,31 @@ public:
 		return Result;
 	}
 
+	virtual void GetDynamicMeshElements(const TArray<const FSceneView*>& Views, const FSceneViewFamily& ViewFamily, uint32 VisibilityMap, FMeshElementCollector& Collector) const override
+	{
+		static const FLinearColor RadiusColor = FLinearColor(FColorList::Orange);
+		for (int32 ViewIndex = 0; ViewIndex < Views.Num(); ViewIndex++)
+		{
+			if (VisibilityMap & (1 << ViewIndex))
+			{
+				FPrimitiveDrawInterface* PDI = Collector.GetPDI(ViewIndex);
+
+				for (auto Location : Locations)
+				{
+					PDI->DrawLine(TestActorLocation, Location, RadiusColor, SDPG_World, 3.f);
+					DrawArrowHead(PDI, Location, TestActorLocation, 30.f, RadiusColor, SDPG_World, 3.f);
+				}
+			}
+		}
+	}
+
 	virtual uint32 GetMemoryFootprint( void ) const { return( sizeof( *this ) + GetAllocatedSize() ); }
-	uint32 GetAllocatedSize( void ) const { return FPrimitiveSceneProxy::GetAllocatedSize(); }
+	uint32 GetAllocatedSize( void ) const { return FPrimitiveSceneProxy::GetAllocatedSize() + Locations.GetAllocatedSize(); }
 
 protected:
 	AFunctionalTest* OwningTest;
+	FVector TestActorLocation;
+	TArray<FVector> Locations;
 };
 //----------------------------------------------------------------------//
 // UFuncTestRenderingComponent
@@ -37,20 +72,43 @@ protected:
 UFuncTestRenderingComponent::UFuncTestRenderingComponent(const FObjectInitializer& ObjectInitializer)
 	: Super(ObjectInitializer)
 {
-	SetCollisionProfileName(UCollisionProfile::NoCollision_ProfileName);
+	// Allows updating in game, while optimizing rendering for the case that it is not modified
+	Mobility = EComponentMobility::Stationary;
+
+	BodyInstance.SetCollisionProfileName(UCollisionProfile::NoCollision_ProfileName);
 
 	AlwaysLoadOnClient = false;
 	AlwaysLoadOnServer = false;
-	bSelectable = false;
+
+	bGenerateOverlapEvents = false;
 }
 
 FPrimitiveSceneProxy* UFuncTestRenderingComponent::CreateSceneProxy()
 {
-	return new FFTestRenderingSceneProxy(this);
+	return new FFTestRenderingSceneProxy(*this);
 }
 
 FBoxSphereBounds UFuncTestRenderingComponent::CalcBounds(const FTransform& LocalToWorld) const
 {
-	FBox BoundingBox = GetOwner()->GetComponentsBoundingBox();
-	return FBoxSphereBounds(BoundingBox);
-}
+	const AFunctionalTest* Owner = Cast<AFunctionalTest>(GetOwner());
+	if (Owner)
+	{
+		FBox BoundingBox;
+
+		TArray<AActor*> RelevantActors;
+		Owner->GatherRelevantActors(RelevantActors);
+
+		BoundingBox += Owner->GetActorLocation();
+		for (auto Actor : RelevantActors)
+		{
+			if (Actor)
+			{
+				BoundingBox += Actor->GetActorLocation();
+			}
+		}
+
+		return FBoxSphereBounds(BoundingBox);
+	}
+
+	return FBoxSphereBounds(EForceInit::ForceInitToZero);
+}
\ No newline at end of file
diff --git a/Engine/Source/Developer/FunctionalTesting/Private/FunctionalAITest.cpp b/Engine/Source/Developer/FunctionalTesting/Private/FunctionalAITest.cpp
--- a/Engine/Source/Developer/FunctionalTesting/Private/FunctionalAITest.cpp
+++ b/Engine/Source/Developer/FunctionalTesting/Private/FunctionalAITest.cpp
@@ -135,6 +135,34 @@ bool AFunctionalAITest::WantsToRunAgain() const
 	return bSingleSetRun == false && CurrentSpawnSetIndex + 1 < SpawnSets.Num();
 }
 
+void AFunctionalAITest::GatherRelevantActors(TArray<AActor*>& OutActors) const
+{
+	Super::GatherRelevantActors(OutActors);
+
+	for (auto SpawnSet : SpawnSets)
+	{
+		if (SpawnSet.FallbackSpawnLocation)
+		{
+			OutActors.AddUnique(SpawnSet.FallbackSpawnLocation);
+		}
+
+		for (auto SpawnInfo : SpawnSet.SpawnInfoContainer)
+		{
+			if (SpawnInfo.SpawnLocation)
+			{
+				OutActors.AddUnique(SpawnInfo.SpawnLocation);
+			}
+		}
+	}
+
+	for (auto Pawn : SpawnedPawns)
+	{
+		if (Pawn)
+		{
+			OutActors.Add(Pawn);
+		}
+	}
+}
 
 void AFunctionalAITest::CleanUp()
 {
diff --git a/Engine/Source/Developer/FunctionalTesting/Private/FunctionalTest.cpp b/Engine/Source/Developer/FunctionalTesting/Private/FunctionalTest.cpp
--- a/Engine/Source/Developer/FunctionalTesting/Private/FunctionalTest.cpp
+++ b/Engine/Source/Developer/FunctionalTesting/Private/FunctionalTest.cpp
@@ -3,6 +3,12 @@
 #include "FunctionalTestingPrivatePCH.h"
 #include "ObjectEditorUtils.h"
 #include "VisualLogger/VisualLogger.h"
+#if WITH_EDITORONLY_DATA
+#include "FuncTestRenderingComponent.h"
+#endif // WITH_EDITORONLY_DATA
+#if WITH_EDITOR
+#include "Engine/Selection.h"
+#endif // WITH_EDITOR
 
 AFunctionalTest::AFunctionalTest( const FObjectInitializer& ObjectInitializer )
 	: Super(ObjectInitializer)
@@ -49,6 +55,21 @@ AFunctionalTest::AFunctionalTest( const FObjectInitializer& ObjectInitializer )
 #endif
 		RootComponent = SpriteComponent;
 	}
+
+#if WITH_EDITORONLY_DATA
+	RenderComp = CreateDefaultSubobject<UFuncTestRenderingComponent>(TEXT("RenderComp"));
+	RenderComp->PostPhysicsComponentTick.bCanEverTick = false;
+	RenderComp->AttachParent = RootComponent;
+#endif // WITH_EDITORONLY_DATA
+
+#if WITH_EDITOR
+	static bool bSelectionHandlerSetUp = false;
+	if (HasAnyFlags(RF_ClassDefaultObject) && !HasAnyFlags(RF_TagGarbageTemp) && bSelectionHandlerSetUp == false)
+	{
+		USelection::SelectObjectEvent.AddStatic(&AFunctionalTest::OnSelectObject);
+		bSelectionHandlerSetUp = true;
+	}
+#endif // WITH_EDITOR
 }
 
 void AFunctionalTest::Tick(float DeltaSeconds)
@@ -197,6 +218,24 @@ void AFunctionalTest::SetTimeLimit(float InTimeLimit, TEnumAsByte<EFunctionalTes
 	TimesUpResult = InResult;
 }
 
+void AFunctionalTest::GatherRelevantActors(TArray<AActor*>& OutActors) const
+{
+	if (ObservationPoint)
+	{
+		OutActors.AddUnique(ObservationPoint);
+	}
+
+	for (auto Actor : AutoDestroyActors)
+	{
+		if (Actor)
+		{
+			OutActors.AddUnique(Actor);
+		}
+	}
+
+	OutActors.Append(DebugGatherRelevantActors());
+}
+
 void AFunctionalTest::RegisterAutoDestroyActor(AActor* ActorToAutoDestroy)
 {
 	AutoDestroyActors.AddUnique(ActorToAutoDestroy);
@@ -235,6 +274,15 @@ void AFunctionalTest::PostEditChangeProperty( struct FPropertyChangedEvent& Prop
 	}
 }
 
+void AFunctionalTest::OnSelectObject(UObject* NewSelection)
+{
+	AFunctionalTest* AsFTest = Cast<AFunctionalTest>(NewSelection);
+	if (AsFTest)
+	{
+		AsFTest->MarkComponentsRenderStateDirty();
+	}
+}
+
 #endif // WITH_EDITOR
 
 void AFunctionalTest::GoToObservationPoint()
