Made functional testing's BP implementable events non virtual on C++ side #UE4

diff --git a/Engine/Config/BaseEngine.ini b/Engine/Config/BaseEngine.ini
--- a/Engine/Config/BaseEngine.ini
+++ b/Engine/Config/BaseEngine.ini
@@ -462,6 +462,10 @@ bUseStreamingPause=false
 +K2FieldRedirects=(OldFieldName="BTTask_MoveDirectlyToward.bForceMoveToLocation",NewFieldName="bDisablePathUpdateOnGoalLocationChange")
 +K2FieldRedirects=(OldFieldName="BTFunctionLibrary.GetBlackboard", NewFieldName="BTFunctionLibrary.GetOwnersBlackboard")
 
+; FunctionalTesting module
++K2FieldRedirects=(OldFieldName="FunctionalTest.WantsToRunAgain", NewFieldName="FunctionalTest.OnWantsReRunCheck")
++K2FieldRedirects=(OldFieldName="FunctionalTest.GetAdditionalTestFinishedMessage", NewFieldName="FunctionalTest.OnAdditionalTestFinishedMessageRequest")
+
 ; Renamed SpriteComponent to BillboardComponent
 +ActiveClassRedirects=(OldClassName="SpriteComponent",NewClassName="BillboardComponent")
 +ActiveClassRedirects=(OldClassName="MaterialSpriteComponent",NewClassName="MaterialBillboardComponent")
diff --git a/Engine/Source/Developer/FunctionalTesting/Classes/FunctionalTest.h b/Engine/Source/Developer/FunctionalTesting/Classes/FunctionalTest.h
--- a/Engine/Source/Developer/FunctionalTesting/Classes/FunctionalTest.h
+++ b/Engine/Source/Developer/FunctionalTesting/Classes/FunctionalTest.h
@@ -101,10 +101,14 @@ public:
 
 	/** retrieves information whether test wants to have another run just after finishing */
 	UFUNCTION(BlueprintImplementableEvent, Category="FunctionalTesting")
-	virtual bool WantsToRunAgain() const;
+	bool OnWantsReRunCheck() const;
+
+	virtual bool WantsToRunAgain() const { return false; }
 
 	UFUNCTION(BlueprintImplementableEvent, Category = "FunctionalTesting")
-	virtual FString GetAdditionalTestFinishedMessage(EFunctionalTestResult::Type TestResult) const;
+	FString OnAdditionalTestFinishedMessageRequest(EFunctionalTestResult::Type TestResult) const;
+	
+	virtual FString GetAdditionalTestFinishedMessage(EFunctionalTestResult::Type TestResult) const { return FString(); }
 	
 	/** ACtors registered this way will be automatically destroyed (by limiting their lifespan)
 	 *	on test finish */
diff --git a/Engine/Source/Developer/FunctionalTesting/Private/FunctionalTest.cpp b/Engine/Source/Developer/FunctionalTesting/Private/FunctionalTest.cpp
--- a/Engine/Source/Developer/FunctionalTesting/Private/FunctionalTest.cpp
+++ b/Engine/Source/Developer/FunctionalTesting/Private/FunctionalTest.cpp
@@ -143,7 +143,7 @@ void AFunctionalTest::FinishTest(TEnumAsByte<EFunctionalTestResult::Type> TestRe
 		, *GetActorLabel()
 		, *ResultText.ToString()
 		, Message.IsEmpty() == false ? *Message : TEXT("Test finished") );
-	const FString AdditionalDetails = GetAdditionalTestFinishedMessage(TestResult) + FString::Printf(TEXT(", time %.2fs"), TotalTime);
+	const FString AdditionalDetails = FString::Printf(TEXT("%s %s, time %.2fs"), *GetAdditionalTestFinishedMessage(TestResult), *OnAdditionalTestFinishedMessageRequest(TestResult), TotalTime);
 
 	AutoDestroyActors.Reset();
 		
diff --git a/Engine/Source/Developer/FunctionalTesting/Private/FunctionalTestingManager.cpp b/Engine/Source/Developer/FunctionalTesting/Private/FunctionalTestingManager.cpp
--- a/Engine/Source/Developer/FunctionalTesting/Private/FunctionalTestingManager.cpp
+++ b/Engine/Source/Developer/FunctionalTesting/Private/FunctionalTestingManager.cpp
@@ -218,7 +218,7 @@ void UFunctionalTestingManager::NotifyTestDone(AFunctionalTest* FTest)
 		GatheredFailedTestsReproString += FTest->GetReproString();
 	}
 
-	if (FTest->WantsToRunAgain() == false)
+	if (FTest->OnWantsReRunCheck() == false && FTest->WantsToRunAgain() == false)
 	{
 		TestsLeft.RemoveSingle(FTest);
 		/*if (bDiscardSuccessfulTests && FTest->IsSuccessful())
