Made FunctionalTest::LogMessage put the message to Message Log as well #UE4

Used to log stuff into vlog alone
UE-17169

diff --git a/Engine/Source/Developer/FunctionalTesting/Private/FunctionalTest.cpp b/Engine/Source/Developer/FunctionalTesting/Private/FunctionalTest.cpp
--- a/Engine/Source/Developer/FunctionalTesting/Private/FunctionalTest.cpp
+++ b/Engine/Source/Developer/FunctionalTesting/Private/FunctionalTest.cpp
@@ -188,7 +188,8 @@ void AFunctionalTest::CleanUp()
 //@todo add "warning" level here
 void AFunctionalTest::LogMessage(const FString& Message)
 {
-	UE_VLOG(this, LogFunctionalTest, Warning
+	UFunctionalTestingManager::AddLogItem(FText::FromString(Message));
+	UE_VLOG(this, LogFunctionalTest, Log
 		, TEXT("%s> %s")
 		, *GetActorLabel(), *Message);
 }
