Fixed FunctionalTestingManager multi-registering for PIE End notification #UE4

While at it added instant-finish tests handling

diff --git a/Engine/Source/Developer/FunctionalTesting/Classes/FunctionalTest.h b/Engine/Source/Developer/FunctionalTesting/Classes/FunctionalTest.h
--- a/Engine/Source/Developer/FunctionalTesting/Classes/FunctionalTest.h
+++ b/Engine/Source/Developer/FunctionalTesting/Classes/FunctionalTest.h
@@ -134,6 +134,7 @@ public:
 	// AActor interface end
 
 	bool IsSuccessful() const { return Result == EFunctionalTestResult::Succeeded; }
+	bool IsRunning() const { return !!bIsRunning; }
 
 protected:
 	void GoToObservationPoint();
@@ -142,7 +143,7 @@ public:
 	FFunctionalTestDoneSignature TestFinishedObserver;
 
 protected:
-	uint32 bIsRunning;
+	uint32 bIsRunning : 1;
 	float TotalTime;
 
 public:
diff --git a/Engine/Source/Developer/FunctionalTesting/Private/FunctionalTestingManager.cpp b/Engine/Source/Developer/FunctionalTesting/Private/FunctionalTestingManager.cpp
--- a/Engine/Source/Developer/FunctionalTesting/Private/FunctionalTestingManager.cpp
+++ b/Engine/Source/Developer/FunctionalTesting/Private/FunctionalTestingManager.cpp
@@ -139,6 +139,7 @@ bool UFunctionalTestingManager::RunAllFunctionalTests(UObject* WorldContext, boo
 
 void UFunctionalTestingManager::OnEndPIE(const bool bIsSimulating)
 {
+	FEditorDelegates::EndPIE.RemoveAll(this);
 	RemoveFromRoot();
 }
 
@@ -275,6 +276,7 @@ void UFunctionalTestingManager::AllTestsDone()
 	else
 	{
 		bFinished = true;
+		FEditorDelegates::EndPIE.RemoveAll(this);
 		AddLogItem(LOCTEXT("TestDone", "DONE."));
 		RemoveFromRoot();
 	}
@@ -335,8 +337,16 @@ bool UFunctionalTestingManager::RunFirstValidTest()
 				TestsLeft[Index]->TestFinishedObserver = TestFinishedObserver;
 				if (TestsLeft[Index]->StartTest())
 				{
-					bTestSuccessfullyTriggered = true;
-					break;
+					if (TestsLeft[Index]->IsRunning() == true)
+					{
+						bTestSuccessfullyTriggered = true;
+						break;
+					}
+					else
+					{
+						// test finished instantly, remove it
+						bRemove = true;
+					}
 				}
 				else
 				{
