Copying //UE4/Dev-Editor to //UE4/Dev-Main (Source: //UE4/Dev-Editor @ 3082391)

diff --git a/Engine/Source/Developer/FunctionalTesting/Classes/FunctionalTest.h b/Engine/Source/Developer/FunctionalTesting/Classes/FunctionalTest.h
--- a/Engine/Source/Developer/FunctionalTesting/Classes/FunctionalTest.h
+++ b/Engine/Source/Developer/FunctionalTesting/Classes/FunctionalTest.h
@@ -115,6 +115,18 @@ struct FUNCTIONALTESTING_API FPerfStatsRecord
 };
 
 
+UENUM(BlueprintType)
+enum class EComparisonMethod : uint8
+{
+	Equal_To,
+	Not_Equal_To,
+	Greater_Than_Or_Equal_To,
+	Less_Than_Or_Equal_To,
+	Greater_Than,
+	Less_Than
+};
+
+
 /** 
  * Class for use with functional tests which provides various performance measuring features. 
  * Recording of basic, unintrusive performance stats.
@@ -211,50 +223,70 @@ public:
 };
 
 UENUM()
-namespace EFunctionalTestResult
+enum class EFunctionalTestResult : uint8
 {
-	enum Type
-	{
-		Invalid,
-		Error,
-		Running,		
-		Failed,
-		Succeeded,
-	};
-}
+	/**
+	 * When finishing a test if you use Default, you're not explicitly stating if the test passed or failed.
+	 * Instead you're a
+	 */
+	Default,
+	Invalid,
+	Error,
+	Running,
+	Failed,
+	Succeeded
+};
 
 DECLARE_DYNAMIC_MULTICAST_DELEGATE(FFunctionalTestEventSignature);
 DECLARE_DELEGATE_OneParam(FFunctionalTestDoneSignature, class AFunctionalTest*);
 
-UCLASS(Blueprintable, MinimalAPI)
-class AFunctionalTest : public AActor
+UCLASS(hidecategories=( Actor, Input, Rendering ), Blueprintable)
+class FUNCTIONALTESTING_API AFunctionalTest : public AActor
 {
-	GENERATED_UCLASS_BODY()
+	GENERATED_BODY()
 
-	static const uint32 DefaultTimeLimit = 60;	// seconds
+public:
+	AFunctionalTest(const FObjectInitializer& ObjectInitializer);
 
 private_subobject:
 	DEPRECATED_FORGAME(4.6, "SpriteComponent should not be accessed directly, please use GetSpriteComponent() function instead. SpriteComponent will soon be private and your code will not compile.")
 	UPROPERTY()
 	UBillboardComponent* SpriteComponent;
+
+protected:
+	UPROPERTY(EditAnywhere, BlueprintReadOnly, Category="Functional Testing")
+	uint32 bIsEnabled:1;
+
+	UPROPERTY(EditAnywhere, BlueprintReadOnly, Category="Functional Testing", meta=( MultiLine="true" ))
+	FString Description;
+
+	UPROPERTY(EditAnywhere, BlueprintReadWrite, Category="Functional Testing")
+	AActor* ObservationPoint;
+
+	UPROPERTY(EditAnywhere, BlueprintReadWrite, Category="Functional Testing", AdvancedDisplay)
+	FRandomStream RandomNumbersStream;
+
 public:
 
-	UPROPERTY(BlueprintReadWrite, Category=FunctionalTesting)
-	TEnumAsByte<EFunctionalTestResult::Type> Result;
+	UPROPERTY(BlueprintReadWrite, Category="Functional Testing")
+	EFunctionalTestResult Result;
 
-	/** If test is limited by time this is the result that will be returned when time runs out */
-	UPROPERTY(EditAnywhere, Category=FunctionalTesting)
-	TEnumAsByte<EFunctionalTestResult::Type> TimesUpResult;
+	/** The Test's time limit for preparation, this is the time it has to return true when checking IsReady(). '0' means no limit. */
+	UPROPERTY(EditAnywhere, BlueprintReadOnly, Category="Timeout")
+	float PreparationTimeLimit;
 
 	/** Test's time limit. '0' means no limit */
-	UPROPERTY(EditAnywhere, BlueprintReadOnly, Category=FunctionalTesting)
+	UPROPERTY(EditAnywhere, BlueprintReadOnly, Category="Timeout")
 	float TimeLimit;
 
-	UPROPERTY(EditAnywhere, BlueprintReadWrite, Category=FunctionalTesting)
+	UPROPERTY(EditAnywhere, BlueprintReadWrite, Category="Timeout", meta=( MultiLine="true" ))
 	FText TimesUpMessage;
 
-	UPROPERTY(EditAnywhere, BlueprintReadWrite, Category=FunctionalTesting)
-	AActor* ObservationPoint;
+	/** If test is limited by time this is the result that will be returned when time runs out */
+	UPROPERTY(EditAnywhere, Category="Timeout")
+	EFunctionalTestResult TimesUpResult;
+
+public:
 
 	/** Called when the test is started */
 	UPROPERTY(BlueprintAssignable)
@@ -268,63 +300,178 @@ public:
 	TArray<AActor*> AutoDestroyActors;
 	
 	FString FailureMessage;
-
-	UPROPERTY(EditAnywhere, BlueprintReadWrite, Category = FunctionalTesting)
-	FRandomStream RandomNumbersStream;
-
-	UPROPERTY(EditAnywhere, BlueprintReadOnly, Category = FunctionalTesting)
-	FString Description;
 	
 #if WITH_EDITORONLY_DATA
 	UPROPERTY()
 	class UFuncTestRenderingComponent* RenderComp;
-#endif // WITH_EDITORONLY_DATA
 
-	UPROPERTY(EditAnywhere, BlueprintReadOnly, Category=FunctionalTesting)
-	uint32 bIsEnabled:1;
+	UPROPERTY()
+	class UTextRenderComponent* TestName;
+#endif // WITH_EDITORONLY_DATA
 
 	/** List of causes we need a re-run. */
 	TArray<FName> RerunCauses;
+
 	/** Cause of the current rerun if we're in a named rerun. */
 	FName CurrentRerunCause;
 
 public:
+	/**
+	 * Assert that a boolean value is true.
+	 * @param Message	The message to display if the assert fails ("Assertion Failed: 'Message' for context ''")
+	 */
+	UFUNCTION(BlueprintCallable, Category = "Asserts", meta = ( HidePin = "ContextObject", DefaultToSelf = "ContextObject"))
+	void AssertTrue(bool Condition, FString Message, const UObject* ContextObject = nullptr);
+
+	/**
+	 * Assert that a boolean value is false.
+	 * @param Message	The message to display if the assert fails ("Assertion Failed: 'Message' for context ''")
+	 */
+	UFUNCTION(BlueprintCallable, Category = "Asserts", meta = ( HidePin = "ContextObject", DefaultToSelf = "ContextObject"))
+	void AssertFalse(bool Condition, FString Message, const UObject* ContextObject = nullptr);
+
+	/**
+	 * Assert that a UObject is valid
+	 * @param Message	The message to display if the object is invalid ("Invalid object: 'Message' for context ''")
+	 */
+	UFUNCTION(BlueprintCallable, Category = "Asserts", meta = ( HidePin = "ContextObject", DefaultToSelf = "ContextObject"))
+	void AssertIsValid(UObject* Object, FString Message, const UObject* ContextObject = nullptr);
+
+	/**
+	 * Assert on a relationship between two integers.
+	 * @param What	A name to use in the message if the assert fails (What: expected {Actual} to be <ShouldBe> {Expected} for context '')
+	 */
+	UFUNCTION(BlueprintCallable, Category = "Asserts", DisplayName = "Assert Value (int)", meta = ( HidePin = "ContextObject", DefaultToSelf = "ContextObject"))
+	void AssertValue_Int(int32 Actual, EComparisonMethod ShouldBe, int32 Expected, const FString& What, const UObject* ContextObject = nullptr);
+
+	/**
+	 * Assert on a relationship between two floats.
+	 * @param What	A name to use in the message if the assert fails (What: expected {Actual} to be <ShouldBe> {Expected} for context '')
+	 */
+	UFUNCTION(BlueprintCallable, Category = "Asserts", DisplayName = "Assert Value (float)", meta = ( HidePin = "ContextObject", DefaultToSelf = "ContextObject"))
+	void AssertValue_Float(float Actual, EComparisonMethod ShouldBe, float Expected, const FString& What, const UObject* ContextObject = nullptr);
+
+	/**
+	 * Assert on a relationship between two DateTimes.
+	 * @param What	A name to use in the message if the assert fails (What: expected {Actual} to be <ShouldBe> {Expected} for context '')
+	 */
+	UFUNCTION(BlueprintCallable, Category = "Asserts", DisplayName = "Assert Value (DateTime)", meta = ( HidePin = "ContextObject", DefaultToSelf = "ContextObject"))
+	void AssertValue_DateTime(FDateTime Actual, EComparisonMethod ShouldBe, FDateTime Expected, const FString& What, const UObject* ContextObject = nullptr);
+
+	/**
+	 * Assert that two transforms are (components memberwise - translation, rotation, scale) equal within a small tolerance.
+	 * @param What	A name to use in the message if the assert fails ("Expected 'What' to be {Expected} but it was {Actual} for context ''")
+	 */
+	UFUNCTION(BlueprintCallable, Category = "Asserts", DisplayName = "Assert Equal (Transform)", meta = ( HidePin = "ContextObject", DefaultToSelf = "ContextObject"))
+	void AssertEqual_Transform(const FTransform Actual, const FTransform Expected, const FString& What, const UObject* ContextObject = nullptr);
+
+	/**
+	 * Assert that two floats are equal within tolerance between two floats.
+	 * @param What	A name to use in the message if the assert fails (What: expected {Actual} to be Equal To {Expected} within Tolerance for context '')
+	 */
+	UFUNCTION(BlueprintCallable, Category = "Asserts", DisplayName = "Assert Equal (Float)", meta = ( HidePin = "ContextObject", DefaultToSelf = "ContextObject"))
+	void AssertEqual_Float(const float Actual, const float Expected, const FString& What, float Tolerance = 1.e-4, const UObject* ContextObject = nullptr);
+
+	/**
+	 * Assert that two transforms are (components memberwise - translation, rotation, scale) not equal within a small tolerance.
+	 * @param What	A name to use in the message if the assert fails ("Expected 'What' not to be {Expected} but it was {Actual} for context ''")
+	 */
+	UFUNCTION(BlueprintCallable, Category = "Asserts", DisplayName = "Assert Not Equal (Transform)", meta = ( HidePin = "ContextObject", DefaultToSelf = "ContextObject"))
+	void AssertNotEqual_Transform(const FTransform Actual, const FTransform NotExpected, const FString& What, const UObject* ContextObject = nullptr);
+
+	/**
+	 * Assert that the component angles of two rotators are all equal within a small tolerance.
+	 * @param What	A name to use in the message if the assert fails ("Expected 'What' to be {Expected} but it was {Actual} for context ''")
+	 */
+	UFUNCTION(BlueprintCallable, Category = "Asserts", DisplayName = "Assert Equal (Rotator)", meta = ( HidePin = "ContextObject", DefaultToSelf = "ContextObject"))
+	void AssertEqual_Rotator(const FRotator Actual, const FRotator Expected, const FString& What, const UObject* ContextObject = nullptr);
+
+	/**
+	 * Assert that the component angles of two rotators are all not equal within a small tolerance.
+	 * @param What	A name to use in the message if the assert fails ("Expected 'What' not to be {Expected} but it was {Actual} for context ''")
+	 */
+	UFUNCTION(BlueprintCallable, Category = "Asserts", DisplayName = "Assert Not Equal (Rotator)", meta = ( HidePin = "ContextObject", DefaultToSelf = "ContextObject"))
+	void AssertNotEqual_Rotator(const FRotator Actual, const FRotator NotExpected, const FString& What, const UObject* ContextObject = nullptr);
+
+	/**
+	 * Assert that two vectors are (memberwise) equal within a small tolerance.
+	 * @param What	A name to use in the message if the assert fails ("Expected 'What' to be {Expected} but it was {Actual} for context ''")
+	 */
+	UFUNCTION(BlueprintCallable, Category = "Asserts", DisplayName = "Assert Equal (Vector)", meta = ( HidePin = "ContextObject", DefaultToSelf = "ContextObject"))
+	void AssertEqual_Vector(const FVector Actual, const FVector Expected, const FString& What, const float Tolerance = 1.e-4f, const UObject* ContextObject = nullptr);
+
+	/**
+	 * Assert that two vectors are (memberwise) not equal within a small tolerance.
+	 * @param What	A name to use in the message if the assert fails ("Expected 'What' not to be {Expected} but it was {Actual} for context ''")
+	 */
+	UFUNCTION(BlueprintCallable, Category = "Asserts", DisplayName = "Assert Not Equal (Vector)", meta = ( HidePin = "ContextObject", DefaultToSelf = "ContextObject"))
+	void AssertNotEqual_Vector(const FVector Actual, const FVector NotExpected, const FString& What, const UObject* ContextObject = nullptr);
+
+	/**
+	 * Assert that two Strings are equal.
+	 * @param What	A name to use in the message if the assert fails ("Expected 'What' to be {Expected} but it was {Actual} for context ''")
+	 */
+	UFUNCTION(BlueprintCallable, Category = "Asserts", DisplayName = "Assert Equal (String)", meta = ( HidePin = "ContextObject", DefaultToSelf = "ContextObject"))
+	void AssertEqual_String(const FString Actual, const FString Expected, const FString& What, const UObject* ContextObject = nullptr);
+
+	/**
+	 * Assert that two Strings are not equal.
+	 * @param What	A name to use in the message if the assert fails ("Expected 'What' not to be {Expected} but it was {Actual} for context ''")
+	 */
+	UFUNCTION(BlueprintCallable, Category = "Asserts", DisplayName = "Assert Not Equal (String)", meta = ( HidePin = "ContextObject", DefaultToSelf = "ContextObject"))
+	void AssertNotEqual_String(const FString Actual, const FString NotExpected, const FString& What, const UObject* ContextObject = nullptr);
+
+protected:
+	void LogAssertFail(const FString& Message);
+	void LogAssertPass(const FString& Message);
 
-	virtual bool StartTest(const TArray<FString>& Params = TArray<FString>());
+public:
+	virtual bool RunTest(const TArray<FString>& Params = TArray<FString>());
 
-	UFUNCTION(BlueprintCallable, Category="Development")
-	virtual void FinishTest(TEnumAsByte<EFunctionalTestResult::Type> TestResult, const FString& Message);
+public:
+	FString	GetCurrentStepName() { return IsInStep() ? CurrentStepName : FString(); }
+	void 	StartStep(const FString& StepName) { CurrentStepName = StepName; bInStep = true; }
+	void 	FinishStep() { CurrentStepName.Empty(); bInStep = false; }
+	bool	IsInStep() { return bInStep; }
 
-	UFUNCTION(BlueprintCallable, Category="Development")
+
+	UFUNCTION(BlueprintCallable, Category="Functional Testing")
+	virtual void FinishTest(EFunctionalTestResult TestResult, const FString& Message);
+
+	UFUNCTION(BlueprintCallable, Category="Functional Testing")
 	virtual void LogMessage(const FString& Message);
 
-	UFUNCTION(BlueprintCallable, Category="Development")
-	virtual void SetTimeLimit(float NewTimeLimit, TEnumAsByte<EFunctionalTestResult::Type> ResultWhenTimeRunsOut);
+	UFUNCTION(BlueprintCallable, Category="Functional Testing")
+	virtual void SetTimeLimit(float NewTimeLimit, EFunctionalTestResult ResultWhenTimeRunsOut);
+
+public:
 
 	/** Used by debug drawing to gather actors this test is using and point at them on the level to better understand test's setup */
-	UFUNCTION(BlueprintImplementableEvent, Category = "FunctionalTesting")
+	UFUNCTION(BlueprintImplementableEvent, Category="Functional Testing")
 	TArray<AActor*> DebugGatherRelevantActors() const;
 
 	virtual void GatherRelevantActors(TArray<AActor*>& OutActors) const;
 
 	/** retrieves information whether test wants to have another run just after finishing */
-	UFUNCTION(BlueprintImplementableEvent, Category="FunctionalTesting")
+	UFUNCTION(BlueprintImplementableEvent, Category="Functional Testing")
 	bool OnWantsReRunCheck() const;
 
 	virtual bool WantsToRunAgain() const { return false; }
 
 	/** Causes the test to be rerun for a specific named reason. */
-	UFUNCTION(BlueprintCallable, Category = "FunctionalTesting")
+	UFUNCTION(BlueprintCallable, Category = "Functional Testing")
 	void AddRerun(FName Reason);
 
 	/** Returns the current re-run reason if we're in a named re-run. */
-	UFUNCTION(BlueprintCallable, Category = "FunctionalTesting")
-	FName GetCurrentRerunReason()const;
+	UFUNCTION(BlueprintCallable, Category = "Functional Testing")
+	FName GetCurrentRerunReason() const;
 
-	UFUNCTION(BlueprintImplementableEvent, Category = "FunctionalTesting")
-	FString OnAdditionalTestFinishedMessageRequest(EFunctionalTestResult::Type TestResult) const;
+	UFUNCTION(BlueprintImplementableEvent, Category = "Functional Testing")
+	FString OnAdditionalTestFinishedMessageRequest(EFunctionalTestResult TestResult) const;
 	
-	virtual FString GetAdditionalTestFinishedMessage(EFunctionalTestResult::Type TestResult) const { return FString(); }
+	virtual FString GetAdditionalTestFinishedMessage(EFunctionalTestResult TestResult) const { return FString(); }
+
+public:
 	
 	/** ACtors registered this way will be automatically destroyed (by limiting their lifespan)
 	 *	on test finish */
@@ -345,25 +492,74 @@ public:
 #endif // WITH_EDITOR
 
 	// AActor interface begin
+	virtual void OnConstruction(const FTransform& Transform) override;
 	virtual void Tick(float DeltaSeconds) override;
 	virtual void EndPlay(const EEndPlayReason::Type EndPlayReason) override;
 	// AActor interface end
 
-	bool IsSuccessful() const { return Result == EFunctionalTestResult::Succeeded; }
-	bool IsRunning() const { return !!bIsRunning; }
+	UFUNCTION(BlueprintCallable, Category = "Functional Testing")
+	bool IsRunning() const;
+
+	UFUNCTION(BlueprintCallable, Category = "Functional Testing")
+	bool IsEnabled() const;
 
 protected:
+	/**
+	 * Prepare Test is fired once the test starts up, before the test IsReady() and thus before Start Test is called.
+	 * So if there's some initial conditions or setup that you might need for your IsReady() check, you might want
+	 * to do that here.
+	 */
+	virtual void PrepareTest();
+
+	/**
+	 * Prepare Test is fired once the test starts up, before the test IsReady() and thus before Start Test is called.
+	 * So if there's some initial conditions or setup that you might need for your IsReady() check, you might want
+	 * to do that here.
+	 */
+	UFUNCTION(BlueprintImplementableEvent, meta=( DisplayName="Prepare Test" ))
+	void ReceivePrepareTest();
+
+	/**
+	 * Called once the IsReady() check for the test returns true.  After that happens the test has Officially started,
+	 * and it will begin receiving Ticks in the blueprint.
+	 */
+	virtual void StartTest();
+	
+	/**
+	 * Called once the IsReady() check for the test returns true.  After that happens the test has Officially started,
+	 * and it will begin receiving Ticks in the blueprint.
+	 */
+	UFUNCTION(BlueprintImplementableEvent, meta=( DisplayName="Start Test" ))
+	void ReceiveStartTest();
+
+	/**
+	 * IsReady() is called once per frame after a test is run, until it returns true.  You should use this function to
+	 * delay Start being called on the test until preconditions are met.
+	 */
+	UFUNCTION(BlueprintNativeEvent, Category="Functional Testing")
+	bool IsReady();
+
+	virtual bool IsReady_Implementation();
+
+	/**
+	 * Goto an observation location.
+	 */
 	void GoToObservationPoint();
 
 public:
 	FFunctionalTestDoneSignature TestFinishedObserver;
 
 protected:
-	UPROPERTY(EditAnywhere, BlueprintReadOnly, Category=FunctionalTesting)
 	bool bIsRunning;
+
+	bool bInStep;
+	FString CurrentStepName;
 	
 	float TotalTime;
 
+private:
+	bool bIsReady;
+
 public:
 	/** Returns SpriteComponent subobject **/
 	UBillboardComponent* GetSpriteComponent();
diff --git a/Engine/Source/Developer/FunctionalTesting/Classes/FunctionalTestGameMode.h b/Engine/Source/Developer/FunctionalTesting/Classes/FunctionalTestGameMode.h
new file mode 100644
--- /dev/null
+++ b/Engine/Source/Developer/FunctionalTesting/Classes/FunctionalTestGameMode.h
@@ -0,0 +1,13 @@
+// Copyright 1998-2016 Epic Games, Inc. All Rights Reserved.
+
+#pragma once
+
+#include "GameFramework/GameMode.h"
+#include "FunctionalTestGameMode.generated.h"
+
+UCLASS(Blueprintable)
+class FUNCTIONALTESTING_API AFunctionalTestGameMode : public AGameMode
+{
+	GENERATED_UCLASS_BODY()
+
+};
diff --git a/Engine/Source/Developer/FunctionalTesting/Classes/FunctionalTestLevelScript.h b/Engine/Source/Developer/FunctionalTesting/Classes/FunctionalTestLevelScript.h
new file mode 100644
--- /dev/null
+++ b/Engine/Source/Developer/FunctionalTesting/Classes/FunctionalTestLevelScript.h
@@ -0,0 +1,12 @@
+// Copyright 1998-2016 Epic Games, Inc. All Rights Reserved.
+#pragma once
+
+#include "Engine/LevelScriptActor.h"
+#include "FunctionalTestLevelScript.generated.h"
+
+UCLASS(notplaceable, meta=(KismetHideOverrides = "ReceiveAnyDamage,ReceivePointDamage,ReceiveRadialDamage,ReceiveActorBeginOverlap,ReceiveActorEndOverlap,ReceiveHit,ReceiveDestroyed,ReceiveActorBeginCursorOver,ReceiveActorEndCursorOver,ReceiveActorOnClicked,ReceiveActorOnReleased,ReceiveActorOnInputTouchBegin,ReceiveActorOnInputTouchEnd,ReceiveActorOnInputTouchEnter,ReceiveActorOnInputTouchLeave"), HideCategories=(Collision,Rendering,"Utilities|Transformation"))
+class FUNCTIONALTESTING_API AFunctionalTestLevelScript : public ALevelScriptActor
+{
+	GENERATED_UCLASS_BODY()
+	
+};
\ No newline at end of file
diff --git a/Engine/Source/Developer/FunctionalTesting/Classes/FunctionalTestingManager.h b/Engine/Source/Developer/FunctionalTesting/Classes/FunctionalTestingManager.h
--- a/Engine/Source/Developer/FunctionalTesting/Classes/FunctionalTestingManager.h
+++ b/Engine/Source/Developer/FunctionalTesting/Classes/FunctionalTestingManager.h
@@ -34,10 +34,12 @@ class UFunctionalTestingManager : public UBlueprintFunctionLibrary
 	UPROPERTY(BlueprintAssignable)
 	FFunctionalTestEventSignature OnTestsBegin;
 	
+	/**
+	 * Triggers in sequence all functional tests found on the level.
+	 * @return true if any tests have been triggered
+	 */
 	UFUNCTION(BlueprintCallable, Category="FunctionalTesting", meta=(WorldContext="WorldContext", CallableWithoutWorldContext ) )
-	/** Triggers in sequence all functional tests found on the level.
-	 *	@return true if any tests have been triggered */
-	 static bool RunAllFunctionalTests(UObject* WorldContext, bool bNewLog = true, bool bRunLooped = false, bool bWaitForNavigationBuildFinish = true, FString FailedTestsReproString = TEXT(""));
+	static bool RunAllFunctionalTests(UObject* WorldContext, bool bNewLog = true, bool bRunLooped = false, bool bWaitForNavigationBuildFinish = true, FString FailedTestsReproString = TEXT(""));
 		
 	bool IsRunning() const { return bIsRunning; }
 	bool IsFinished() const { return bFinished; }
@@ -76,7 +78,7 @@ protected:
 
 	FFunctionalTestDoneSignature TestFinishedObserver;
 	
-	FString GatheredFailedTestsReproString;
+	//FString GatheredFailedTestsReproString;
 	FString StartingReproString;
 	TArray<FString> TestReproStrings;
 
diff --git a/Engine/Source/Developer/FunctionalTesting/FunctionalTesting.Build.cs b/Engine/Source/Developer/FunctionalTesting/FunctionalTesting.Build.cs
--- a/Engine/Source/Developer/FunctionalTesting/FunctionalTesting.Build.cs
+++ b/Engine/Source/Developer/FunctionalTesting/FunctionalTesting.Build.cs
@@ -15,6 +15,7 @@ public class FunctionalTesting : ModuleRules
                 "MessageLog",
                 "AIModule",
                 "RenderCore",
+                "AssetRegistry",
                 "RHI"
 			}
 			);
@@ -26,6 +27,6 @@ public class FunctionalTesting : ModuleRules
 				"Stats/Public",
 				"Developer/FunctionalTesting/Private",
 			}
-			);
+		);
 	}
 }
diff --git a/Engine/Source/Developer/FunctionalTesting/Private/AutomationBlueprintFunctionLibrary.cpp b/Engine/Source/Developer/FunctionalTesting/Private/AutomationBlueprintFunctionLibrary.cpp
new file mode 100644
--- /dev/null
+++ b/Engine/Source/Developer/FunctionalTesting/Private/AutomationBlueprintFunctionLibrary.cpp
@@ -0,0 +1,103 @@
+// Copyright 1998-2016 Epic Games, Inc. All Rights Reserved.
+#include "FunctionalTestingPrivatePCH.h"
+
+#include "AutomationCommon.h"
+#include "AutomationTest.h"
+#include "DelayForFramesLatentAction.h"
+#include "TakeScreenshotAfterTimeLatentAction.h"
+#include "Engine/LatentActionManager.h"
+#include "SlateBasics.h"
+
+#include "AutomationBlueprintFunctionLibrary.h"
+
+DEFINE_LOG_CATEGORY_STATIC(BlueprintAssertion, Error, Error)
+DEFINE_LOG_CATEGORY_STATIC(AutomationFunctionLibrary, Log, Log)
+
+UAutomationBlueprintFunctionLibrary::UAutomationBlueprintFunctionLibrary(const class FObjectInitializer& Initializer) : Super(Initializer)
+{
+}
+
+void UAutomationBlueprintFunctionLibrary::TakeAutomationScreenshotInternal(const FString& Name)
+{
+	// Force all mip maps to load before taking the screenshot.
+	UTexture::ForceUpdateTextureStreaming();
+
+	const bool bAddFilenameSuffix = true;
+	FScreenshotRequest::RequestScreenshot(Name, false, bAddFilenameSuffix);
+}
+
+void UAutomationBlueprintFunctionLibrary::TakeAutomationScreenshot(UObject* WorldContextObject, FLatentActionInfo LatentInfo, const FString& Name)
+{
+	if (GIsAutomationTesting)
+	{
+		if(!FAutomationTestFramework::GetInstance().IsScreenshotAllowed())
+		{
+			UE_LOG(AutomationFunctionLibrary, Error, TEXT("Attempted to capture screenshot (%s) but screenshots are not enabled. If you're seeing this on TeamCity, please make sure the name of your test contains '_VisualTest'"), *Name);
+			return;
+		}
+
+		if (UWorld* World = GEngine->GetWorldFromContextObject(WorldContextObject))
+		{
+			FLatentActionManager& LatentActionManager = World->GetLatentActionManager();
+			if (LatentActionManager.FindExistingAction<FTakeScreenshotAfterTimeLatentAction>(LatentInfo.CallbackTarget, LatentInfo.UUID) == nullptr)
+			{
+				LatentActionManager.AddNewAction(LatentInfo.CallbackTarget, LatentInfo.UUID, new FTakeScreenshotAfterTimeLatentAction(LatentInfo, Name, 0.0f));
+			}
+		}
+	}
+	else
+	{
+		UE_LOG(AutomationFunctionLibrary, Log, TEXT("Screenshot not captured - screenshots are only taken during automation tests"));
+	}
+}
+
+void UAutomationBlueprintFunctionLibrary::TakeAutomationScreenshotAtCamera(UObject* WorldContextObject, FLatentActionInfo LatentInfo, ACameraActor* Camera, const FString& NameOverride, float DelayBeforeScreenshotSeconds)
+{
+	if ( Camera == nullptr )
+	{
+		UE_LOG(AutomationFunctionLibrary, Error, TEXT("A camera is required to TakeAutomationScreenshotAtCamera"));
+		return;
+	}
+
+	APlayerController* PlayerController = UGameplayStatics::GetPlayerController(WorldContextObject, 0);
+
+	if ( PlayerController == nullptr )
+	{
+		UE_LOG(AutomationFunctionLibrary, Error, TEXT("A player controller is required to TakeAutomationScreenshotAtCamera"));
+		return;
+	}
+
+	APawn* PlayerPawn = UGameplayStatics::GetPlayerPawn(WorldContextObject, 0);
+	FTransform CameraTransform = Camera->GetTransform();
+
+	// Move the player, then queue up a screenshot.
+	// We need to delay before the screenshot so that the motion blur has time to stop.
+	PlayerController->SetViewTarget(Camera, FViewTargetTransitionParams());
+	FString ScreenshotName = Camera->GetName();
+
+	if ( !NameOverride.IsEmpty() )
+	{
+		ScreenshotName = NameOverride;
+	}
+
+	if ( UWorld* World = GEngine->GetWorldFromContextObject(WorldContextObject) )
+	{
+		ScreenshotName = FString::Printf(TEXT("%s_%s"), *World->GetName(), *ScreenshotName);
+
+		FLatentActionManager& LatentActionManager = World->GetLatentActionManager();
+		if ( LatentActionManager.FindExistingAction<FTakeScreenshotAfterTimeLatentAction>(LatentInfo.CallbackTarget, LatentInfo.UUID) == nullptr )
+		{
+			LatentActionManager.AddNewAction(LatentInfo.CallbackTarget, LatentInfo.UUID, new FTakeScreenshotAfterTimeLatentAction(LatentInfo, ScreenshotName, DelayBeforeScreenshotSeconds));
+		}
+	}
+}
+
+void UAutomationBlueprintFunctionLibrary::BeginPerformanceCapture()
+{
+    //::BeginPerformanceCapture();
+}
+
+void UAutomationBlueprintFunctionLibrary::EndPerformanceCapture()
+{
+    //::EndPerformanceCapture();
+}
diff --git a/Engine/Source/Developer/FunctionalTesting/Private/FuncTestManager.cpp b/Engine/Source/Developer/FunctionalTesting/Private/FuncTestManager.cpp
--- a/Engine/Source/Developer/FunctionalTesting/Private/FuncTestManager.cpp
+++ b/Engine/Source/Developer/FunctionalTesting/Private/FuncTestManager.cpp
@@ -1,10 +1,49 @@
 // Copyright 1998-2016 Epic Games, Inc. All Rights Reserved.
 
 #include "FunctionalTestingPrivatePCH.h"
+
+#include "UObject/UObject.h"
 #include "Engine/World.h"
 
 DEFINE_LOG_CATEGORY(LogFunctionalTest);
 
+FFuncTestManager::FFuncTestManager()
+{
+#if WITH_EDITOR
+	FWorldDelegates::GetAssetTags.AddRaw(this, &FFuncTestManager::OnGetAssetTagsForWorld);
+#endif
+}
+
+FFuncTestManager::~FFuncTestManager()
+{
+#if WITH_EDITOR
+	FWorldDelegates::GetAssetTags.RemoveAll(this);
+#endif
+}
+
+void FFuncTestManager::OnGetAssetTagsForWorld(const UWorld* World, TArray<UObject::FAssetRegistryTag>& OutTags)
+{
+#if WITH_EDITOR
+	int32 Tests = 0;
+	FString TestNames;
+	for ( TActorIterator<AFunctionalTest> ActorItr(const_cast<UWorld*>( World )); ActorItr; ++ActorItr )
+	{
+		AFunctionalTest* FunctionalTest = *ActorItr;
+
+		// Only include enabled tests in the list of functional tests to run.
+		if ( FunctionalTest->IsEnabled() )
+		{
+			Tests++;
+			TestNames.Append(FunctionalTest->GetActorLabel() + TEXT("|") + FunctionalTest->GetName());
+			TestNames.Append(TEXT(";"));
+		}
+	}
+
+	OutTags.Add(UObject::FAssetRegistryTag("Tests", FString::FromInt(Tests), UObject::FAssetRegistryTag::TT_Numerical));
+	OutTags.Add(UObject::FAssetRegistryTag("TestNames", TestNames, UObject::FAssetRegistryTag::TT_Hidden));
+#endif
+}
+
 void FFuncTestManager::SetScript(class UFunctionalTestingManager* NewScript)
 {
 	TestScript = NewScript;
@@ -28,20 +67,21 @@ void FFuncTestManager::SetLooping(const bool bLoop)
 	}
 }
 
-void FFuncTestManager::RunAllTestsOnMap(bool bClearLog, bool bRunLooped)
+UWorld* FFuncTestManager::GetTestWorld()
 {
-	UWorld* TestWorld = NULL;
+	UWorld* TestWorld = nullptr;
+
 #if WITH_EDITOR
 	const TIndirectArray<FWorldContext>& WorldContexts = GEngine->GetWorldContexts();
-	for (const FWorldContext& Context : WorldContexts)
+	for ( const FWorldContext& Context : WorldContexts )
 	{
-		if ((Context.WorldType == EWorldType::PIE) && (Context.World() != NULL))
+		if ( ( Context.WorldType == EWorldType::PIE ) && ( Context.World() != nullptr ) )
 		{
 			TestWorld = Context.World();
 		}
 	}
 #endif
-	if (!TestWorld)
+	if ( !TestWorld )
 	{
 		TestWorld = GWorld;
 		if (GIsEditor)
@@ -50,9 +90,25 @@ void FFuncTestManager::RunAllTestsOnMap(bool bClearLog, bool bRunLooped)
 		}
 	}
 
-	if (TestWorld)
+	return TestWorld;
+}
+
+void FFuncTestManager::RunAllTestsOnMap(bool bClearLog, bool bRunLooped)
+{
+	if ( UWorld* TestWorld = GetTestWorld() )
+	{
+		if ( UFunctionalTestingManager::RunAllFunctionalTests(TestWorld, bClearLog, bRunLooped) == false )
+		{
+			UE_LOG(LogFunctionalTest, Error, TEXT("No functional testing script on map."));
+		}
+	}
+}
+
+void FFuncTestManager::RunTestOnMap(const FString& TestName, bool bClearLog, bool bRunLooped)
+{
+	if ( UWorld* TestWorld = GetTestWorld() )
 	{
-		if (UFunctionalTestingManager::RunAllFunctionalTests(TestWorld, bClearLog, bRunLooped) == false)
+		if ( UFunctionalTestingManager::RunAllFunctionalTests(TestWorld, bClearLog, bRunLooped, true, TestName) == false )
 		{
 			UE_LOG(LogFunctionalTest, Error, TEXT("No functional testing script on map."));
 		}
diff --git a/Engine/Source/Developer/FunctionalTesting/Private/FuncTestManager.h b/Engine/Source/Developer/FunctionalTesting/Private/FuncTestManager.h
--- a/Engine/Source/Developer/FunctionalTesting/Private/FuncTestManager.h
+++ b/Engine/Source/Developer/FunctionalTesting/Private/FuncTestManager.h
@@ -8,10 +8,13 @@ class FFuncTestManager : public IFuncTestManager, public TSharedFromThis<FFuncTe
 {
 public:
 
-	virtual ~FFuncTestManager() {}
+	FFuncTestManager();
+	virtual ~FFuncTestManager();
 		
 	/** Triggers in sequence all functional tests found on the level.*/
 	virtual void RunAllTestsOnMap(bool bClearLog, bool bRunLooped) override;
+
+	virtual void RunTestOnMap(const FString& TestName, bool bClearLog, bool bRunLooped) override;
 	
 	virtual bool IsRunning() const override;
 
@@ -23,6 +26,11 @@ public:
 
 	virtual void SetLooping(const bool bLoop) override;
 
+protected:
+	UWorld* GetTestWorld();
+
+	void OnGetAssetTagsForWorld(const UWorld* World, TArray<UObject::FAssetRegistryTag>& OutTags);
+
 protected:
 
 	TWeakObjectPtr<class UFunctionalTestingManager> TestScript;
diff --git a/Engine/Source/Developer/FunctionalTesting/Private/FunctionalAITest.cpp b/Engine/Source/Developer/FunctionalTesting/Private/FunctionalAITest.cpp
--- a/Engine/Source/Developer/FunctionalTesting/Private/FunctionalAITest.cpp
+++ b/Engine/Source/Developer/FunctionalTesting/Private/FunctionalAITest.cpp
@@ -51,7 +51,7 @@ void AFunctionalAITest::BeginPlay()
 	Super::BeginPlay();
 }
 
-bool AFunctionalAITest::StartTest(const TArray<FString>& Params)
+bool AFunctionalAITest::RunTest(const TArray<FString>& Params)
 {
 	KillOffSpawnedPawns();
 	ClearPendingDelayedSpawns();
@@ -75,7 +75,7 @@ bool AFunctionalAITest::StartTest(const TArray<FString>& Params)
 	
 	StartSpawning();
 
-	return Super::StartTest(Params);
+	return Super::RunTest(Params);
 }
 
 void AFunctionalAITest::StartSpawning()
@@ -202,7 +202,7 @@ void AFunctionalAITest::CleanUp()
 	ClearPendingDelayedSpawns();
 }
 
-FString AFunctionalAITest::GetAdditionalTestFinishedMessage(EFunctionalTestResult::Type TestResult) const
+FString AFunctionalAITest::GetAdditionalTestFinishedMessage(EFunctionalTestResult TestResult) const
 {
 	FString ResultStr;
 
diff --git a/Engine/Source/Developer/FunctionalTesting/Private/FunctionalTest.cpp b/Engine/Source/Developer/FunctionalTesting/Private/FunctionalTest.cpp
--- a/Engine/Source/Developer/FunctionalTesting/Private/FunctionalTest.cpp
+++ b/Engine/Source/Developer/FunctionalTesting/Private/FunctionalTest.cpp
@@ -10,18 +10,85 @@
 #include "Engine/Selection.h"
 #endif // WITH_EDITOR
 
+#include "DelayForFramesLatentAction.h"
+
+namespace
+{
+	template <typename T>
+	bool PerformComparison(const T& lhs, const T& rhs, EComparisonMethod comparison)
+	{
+		switch (comparison)
+		{
+		case EComparisonMethod::Equal_To:
+			return lhs == rhs;
+
+		case EComparisonMethod::Not_Equal_To:
+			return lhs != rhs;
+
+		case EComparisonMethod::Greater_Than_Or_Equal_To:
+			return lhs >= rhs;
+
+		case EComparisonMethod::Less_Than_Or_Equal_To:
+			return lhs <= rhs;
+
+		case EComparisonMethod::Greater_Than:
+			return lhs > rhs;
+
+		case EComparisonMethod::Less_Than:
+			return lhs < rhs;
+		}
+
+		UE_LOG(LogFunctionalTest, Error, TEXT("Invalid comparison method"));
+		return false;
+	}
+
+	FString GetComparisonAsString(EComparisonMethod comparison)
+	{
+		UEnum* Enum = FindObject<UEnum>(ANY_PACKAGE, TEXT("EComparisonMethod"), true);
+		return Enum->GetEnumName((uint8)comparison).ToLower().Replace(TEXT("_"), TEXT(" "));
+	}
+
+	FString TransformToString(const FTransform &transform)
+	{
+		const FRotator R(transform.Rotator());
+		FVector T(transform.GetTranslation());
+		FVector S(transform.GetScale3D());
+
+		return FString::Printf(TEXT("Translation: %f, %f, %f | Rotation: %f, %f, %f | Scale: %f, %f, %f"), T.X, T.Y, T.Z, R.Pitch, R.Yaw, R.Roll, S.X, S.Y, S.Z);
+	}
+
+	void DelayForFramesCommon(UObject* WorldContextObject, FLatentActionInfo LatentInfo, int32 NumFrames)
+	{
+		if (UWorld* World = GEngine->GetWorldFromContextObject(WorldContextObject))
+		{
+			FLatentActionManager& LatentActionManager = World->GetLatentActionManager();
+			if (LatentActionManager.FindExistingAction<FDelayForFramesLatentAction>(LatentInfo.CallbackTarget, LatentInfo.UUID) == nullptr)
+			{
+				LatentActionManager.AddNewAction(LatentInfo.CallbackTarget, LatentInfo.UUID, new FDelayForFramesLatentAction(LatentInfo, NumFrames));
+			}
+		}
+	}
+}
+
 AFunctionalTest::AFunctionalTest( const FObjectInitializer& ObjectInitializer )
 	: Super(ObjectInitializer)
-	, TimesUpResult(EFunctionalTestResult::Failed)
-	, TimeLimit(DefaultTimeLimit)
-	, TimesUpMessage( NSLOCTEXT("FunctionalTest", "DefaultTimesUpMessage", "Time's up!") )
 	, bIsEnabled(true)
+	, Result(EFunctionalTestResult::Invalid)
+	, PreparationTimeLimit(15.0f)
+	, TimeLimit(60.0f)
+	, TimesUpMessage( NSLOCTEXT("FunctionalTest", "DefaultTimesUpMessage", "Time's up!") )
+	, TimesUpResult(EFunctionalTestResult::Failed)
 	, bIsRunning(false)
+	, bInStep(false)
+	, CurrentStepName()
 	, TotalTime(0.f)
+	, bIsReady(false)
 {
 	PrimaryActorTick.bCanEverTick = true;
 	PrimaryActorTick.bStartWithTickEnabled = false;
-		
+	
+	bCanBeDamaged = false;
+
 	SpriteComponent = CreateDefaultSubobject<UBillboardComponent>(TEXT("Sprite"));
 	if (SpriteComponent)
 	{
@@ -68,31 +135,28 @@ AFunctionalTest::AFunctionalTest( const FObjectInitializer& ObjectInitializer )
 		bSelectionHandlerSetUp = true;
 	}
 #endif // WITH_EDITOR
+
+#if WITH_EDITORONLY_DATA
+	TestName = CreateDefaultSubobject<UTextRenderComponent>(TEXT("TestName"));
+
+	TestName->SetHorizontalAlignment(EHTA_Center);
+	TestName->SetTextRenderColor(FColor(11, 255, 0));
+	TestName->SetRelativeLocation(FVector(0, 0, 80));
+	TestName->PostPhysicsComponentTick.bCanEverTick = false;
+	TestName->SetupAttachment(RootComponent);
+#endif
 }
 
-void AFunctionalTest::Tick(float DeltaSeconds)
+void AFunctionalTest::OnConstruction(const FTransform& Transform)
 {
-	// already requested not to tick. 
-	if (bIsRunning == false)
-	{
-		return;
-	}
+	Super::OnConstruction(Transform);
 
-	//Do not collect garbage during the test. We force GC at the end.
-	GetWorld()->DelayGarbageCollection();
-
-	TotalTime += DeltaSeconds;
-	if (TimeLimit > 0.f && TotalTime > TimeLimit)
-	{
-		FinishTest(TimesUpResult, TimesUpMessage.ToString());
-	}
-	else
-	{
-		Super::Tick(DeltaSeconds);
-	}
+#if WITH_EDITOR
+	TestName->SetText(FText::FromString(GetActorLabel()));
+#endif
 }
 
-bool AFunctionalTest::StartTest(const TArray<FString>& Params)
+bool AFunctionalTest::RunTest(const TArray<FString>& Params)
 {
 	FailureMessage = TEXT("");
 	
@@ -100,23 +164,84 @@ bool AFunctionalTest::StartTest(const TArray<FString>& Params)
 	GetWorld()->DelayGarbageCollection();
 
 	TotalTime = 0.f;
-	if (TimeLimit > 0)
+	if (TimeLimit >= 0)
 	{
 		SetActorTickEnabled(true);
 	}
 
+	bIsReady = false;
 	bIsRunning = true;
 
 	GoToObservationPoint();
-	
+
+	return true;
+}
+
+void AFunctionalTest::PrepareTest()
+{
+	ReceivePrepareTest();
+}
+
+void AFunctionalTest::StartTest()
+{
+	TotalTime = 0.f;
+
+	ReceiveStartTest();
 	OnTestStart.Broadcast();
+}
+
+void AFunctionalTest::Tick(float DeltaSeconds)
+{
+	// already requested not to tick. 
+	if ( bIsRunning == false )
+	{
+		return;
+	}
 
+	//Do not collect garbage during the test. We force GC at the end.
+	GetWorld()->DelayGarbageCollection();
+
+	if ( !bIsReady )
+	{
+		bIsReady = IsReady();
+
+		// Once we're finally ready to begin the test, then execute the Start event.
+		if ( bIsReady )
+		{
+			StartTest();
+		}
+	}
+
+	if ( bIsReady )
+	{
+		TotalTime += DeltaSeconds;
+		if ( TimeLimit > 0.f && TotalTime > TimeLimit )
+		{
+			FinishTest(TimesUpResult, TimesUpMessage.ToString());
+		}
+		else
+		{
+			Super::Tick(DeltaSeconds);
+		}
+	}
+	else
+	{
+		TotalTime += DeltaSeconds;
+		if ( PreparationTimeLimit > 0.f && TotalTime > PreparationTimeLimit )
+		{
+			FinishTest(TimesUpResult, TimesUpMessage.ToString());
+		}
+	}
+}
+
+bool AFunctionalTest::IsReady_Implementation()
+{
 	return true;
 }
 
-void AFunctionalTest::FinishTest(TEnumAsByte<EFunctionalTestResult::Type> TestResult, const FString& Message)
+void AFunctionalTest::FinishTest(EFunctionalTestResult TestResult, const FString& Message)
 {
-	const static UEnum* FTestResultTypeEnum = FindObject<UEnum>( NULL, TEXT("FunctionalTesting.EFunctionalTestResult") );
+	const static UEnum* FTestResultTypeEnum = FindObject<UEnum>( nullptr, TEXT("FunctionalTesting.EFunctionalTestResult") );
 	
 	if (bIsRunning == false)
 	{
@@ -145,7 +270,7 @@ void AFunctionalTest::FinishTest(TEnumAsByte<EFunctionalTestResult::Type> TestRe
 		}
 	}
 
-	const FText ResultText = FTestResultTypeEnum->GetEnumText( TestResult.GetValue() );
+	const FText ResultText = FTestResultTypeEnum->GetEnumText( (int32)TestResult );
 	const FString OutMessage = FString::Printf(TEXT("%s %s: \"%s\'")
 		, *GetName()
 		, *ResultText.ToString()
@@ -154,7 +279,7 @@ void AFunctionalTest::FinishTest(TEnumAsByte<EFunctionalTestResult::Type> TestRe
 
 	AutoDestroyActors.Reset();
 		
-	switch (TestResult.GetValue())
+	switch (TestResult)
 	{
 		case EFunctionalTestResult::Invalid:
 		case EFunctionalTestResult::Error:
@@ -194,6 +319,16 @@ void AFunctionalTest::CleanUp()
 	FailureMessage = TEXT("");
 }
 
+bool AFunctionalTest::IsRunning() const
+{
+	return bIsRunning;
+}
+
+bool AFunctionalTest::IsEnabled() const
+{
+	return bIsEnabled;
+}
+
 //@todo add "warning" level here
 void AFunctionalTest::LogMessage(const FString& Message)
 {
@@ -203,7 +338,7 @@ void AFunctionalTest::LogMessage(const FString& Message)
 		, *GetName(), *Message);
 }
 
-void AFunctionalTest::SetTimeLimit(float InTimeLimit, TEnumAsByte<EFunctionalTestResult::Type> InResult)
+void AFunctionalTest::SetTimeLimit(float InTimeLimit, EFunctionalTestResult InResult)
 {
 	if (InTimeLimit < 0.f)
 	{
@@ -305,7 +440,7 @@ void AFunctionalTest::OnSelectObject(UObject* NewSelection)
 
 void AFunctionalTest::GoToObservationPoint()
 {
-	if (ObservationPoint == NULL)
+	if (ObservationPoint == nullptr)
 	{
 		return;
 	}
@@ -323,8 +458,220 @@ void AFunctionalTest::GoToObservationPoint()
 }
 
 /** Returns SpriteComponent subobject **/
-UBillboardComponent* AFunctionalTest::GetSpriteComponent() { return SpriteComponent; }
+UBillboardComponent* AFunctionalTest::GetSpriteComponent()
+{
+	return SpriteComponent;
+}
+
+//////////////////////////////////////////////////////////////////////////
+
+void AFunctionalTest::AssertTrue(bool Condition, FString Message, const UObject* ContextObject)
+{
+	if ( !Condition )
+	{
+		LogAssertFail(FString::Printf(TEXT("Assertion failed: '%s' for context '%s'"), *Message, ContextObject ? *ContextObject->GetName() : TEXT("")));
+	}
+	else
+	{
+		LogAssertPass(FString::Printf(TEXT("Assertion passed (%s)"), *Message));
+	}
+}
+
+void AFunctionalTest::AssertFalse(bool Condition, FString Message, const UObject* ContextObject)
+{
+	AssertTrue(!Condition, Message, ContextObject);
+}
 
+void AFunctionalTest::AssertIsValid(UObject* Object, FString Message, const UObject* ContextObject)
+{
+	if ( !IsValid(Object) )
+	{
+		LogAssertFail(FString::Printf(TEXT("Invalid object: '%s' for context '%s'"), *Message, ContextObject ? *ContextObject->GetName() : TEXT("")));
+	}
+	else
+	{
+		LogAssertPass(FString::Printf(TEXT("Valid object: (%s)"), *Message));
+	}
+}
+
+void AFunctionalTest::AssertValue_Int(int32 Actual, EComparisonMethod ShouldBe, int32 Expected, const FString& What, const UObject* ContextObject)
+{
+	if ( !PerformComparison(Actual, Expected, ShouldBe) )
+	{
+		LogAssertFail(FString::Printf(TEXT("%s: expected {%d} to be %s {%d} for context '%s'"), *What, Actual, *GetComparisonAsString(ShouldBe), Expected, ContextObject ? *ContextObject->GetName() : TEXT("")));
+	}
+	else
+	{
+		LogAssertPass(FString::Printf(TEXT("Int assertion passed (%s)"), *What));
+	}
+}
+
+void AFunctionalTest::AssertValue_Float(float Actual, EComparisonMethod ShouldBe, float Expected, const FString& What, const UObject* ContextObject)
+{
+	if ( !PerformComparison(Actual, Expected, ShouldBe) )
+	{
+		LogAssertFail(FString::Printf(TEXT("%s: expected {%f} to be %s {%f} for context '%s'"), *What, Actual, *GetComparisonAsString(ShouldBe), Expected, ContextObject ? *ContextObject->GetName() : TEXT("")));
+	}
+	else
+	{
+		LogAssertPass(FString::Printf(TEXT("Float assertion passed (%s)"), *What));
+	}
+}
+
+void AFunctionalTest::AssertValue_DateTime(FDateTime Actual, EComparisonMethod ShouldBe, FDateTime Expected, const FString& What, const UObject* ContextObject)
+{
+	if ( !PerformComparison(Actual, Expected, ShouldBe) )
+	{
+		LogAssertFail(FString::Printf(TEXT("%s: expected {%s} to be %s {%s} for context '%s'"), *What, *Actual.ToString(), *GetComparisonAsString(ShouldBe), *Expected.ToString(), ContextObject ? *ContextObject->GetName() : TEXT("")));
+	}
+	else
+	{
+		LogAssertPass(FString::Printf(TEXT("DateTime assertion passed (%s)"), *What));
+	}
+}
+
+
+void AFunctionalTest::AssertEqual_Float(const float Actual, const float Expected, const FString& What, const float Tolerance, const UObject* ContextObject)
+{
+	if ( !FMath::IsNearlyEqual(Actual, Expected, Tolerance) )
+	{
+		LogAssertFail(FString::Printf(TEXT("Expected '%s' to be {%f}, but it was {%f} within tolerance {%f} for context '%s'"), *What, Expected, Actual, Tolerance, ContextObject ? *ContextObject->GetName() : TEXT("")));
+	}
+	else
+	{
+		LogAssertPass(FString::Printf(TEXT("Float assertion passed (%s)"), *What));
+	}
+}
+
+void AFunctionalTest::AssertEqual_Transform(const FTransform Actual, const FTransform Expected, const FString& What, const UObject* ContextObject)
+{
+	if ( !Expected.Equals(Actual) )
+	{
+		LogAssertFail(FString::Printf(TEXT("Expected '%s' to be {%s}, but it was {%s} for context '%s'"), *What, *TransformToString(Expected), *TransformToString(Actual), ContextObject ? *ContextObject->GetName() : TEXT("")));
+	}
+	else
+	{
+		LogAssertPass(FString::Printf(TEXT("Transform assertion passed (%s)"), *What));
+	}
+}
+
+void AFunctionalTest::AssertNotEqual_Transform(const FTransform Actual, const FTransform NotExpected, const FString& What, const UObject* ContextObject)
+{
+	if ( NotExpected.Equals(Actual) )
+	{
+		LogAssertFail(FString::Printf(TEXT("Expected '%s' not to be {%s} for context '%s'"), *What, *TransformToString(NotExpected), ContextObject ? *ContextObject->GetName() : TEXT("")));
+	}
+	else
+	{
+		LogAssertPass(FString::Printf(TEXT("Transform assertion passed (%s)"), *What));
+	}
+}
+
+void AFunctionalTest::AssertEqual_Rotator(const FRotator Actual, const FRotator Expected, const FString& What, const UObject* ContextObject)
+{
+	if ( !Expected.Equals(Actual) )
+	{
+		LogAssertFail(FString::Printf(TEXT("Expected '%s' to be {%s} but it was {%s} for context '%s'"), *What, *Expected.ToString(), *Actual.ToString(), ContextObject ? *ContextObject->GetName() : TEXT("")));
+	}
+	else
+	{
+		LogAssertPass(FString::Printf(TEXT("Rotator assertion passed (%s)"), *What));
+	}
+}
+
+void AFunctionalTest::AssertNotEqual_Rotator(const FRotator Actual, const FRotator NotExpected, const FString& What, const UObject* ContextObject)
+{
+	if ( NotExpected.Equals(Actual) )
+	{
+		LogAssertFail(FString::Printf(TEXT("Expected '%s' not to be {%s} for context '%s'"), *What, *NotExpected.ToString(), ContextObject ? *ContextObject->GetName() : TEXT("")));
+	}
+	else
+	{
+		LogAssertPass(FString::Printf(TEXT("Rotator assertion passed (%s)"), *What));
+	}
+}
+
+void AFunctionalTest::AssertEqual_Vector(const FVector Actual, const FVector Expected, const FString& What, const float Tolerance, const UObject* ContextObject)
+{
+	if ( !Expected.Equals(Actual, Tolerance) )
+	{
+		LogAssertFail(FString::Printf(TEXT("Expected '%s' to be {%s} but it was {%s} within tolerance {%f} for context '%s'"), *What, *Expected.ToString(), *Actual.ToString(), Tolerance, ContextObject ? *ContextObject->GetName() : TEXT("")));
+	}
+	else
+	{
+		LogAssertPass(FString::Printf(TEXT("Vector assertion passed (%s)"), *What));
+	}
+}
+
+void AFunctionalTest::AssertNotEqual_Vector(const FVector Actual, const FVector NotExpected, const FString& What, const UObject* ContextObject)
+{
+	if ( NotExpected.Equals(Actual) )
+	{
+		LogAssertFail(FString::Printf(TEXT("Expected '%s' not to be {%s} for context '%s'"), *What, *NotExpected.ToString(), ContextObject ? *ContextObject->GetName() : TEXT("")));
+	}
+	else
+	{
+		LogAssertPass(FString::Printf(TEXT("Vector assertion passed (%s)"), *What));
+	}
+}
+
+void AFunctionalTest::AssertEqual_String(const FString Actual, const FString Expected, const FString& What, const UObject* ContextObject)
+{
+	if ( !Expected.Equals(Actual) )
+	{
+		LogAssertFail(FString::Printf(TEXT("Expected '%s' to be {%s} but it was {%s} for context '%s'"), *What, *Expected, *Actual, ContextObject ? *ContextObject->GetName() : TEXT("")));
+	}
+	else
+	{
+		LogAssertPass(FString::Printf(TEXT("String assertion passed (%s)"), *What));
+	}
+}
+
+void AFunctionalTest::AssertNotEqual_String(const FString Actual, const FString NotExpected, const FString& What, const UObject* ContextObject)
+{
+	if ( NotExpected.Equals(Actual) )
+	{
+		LogAssertFail(FString::Printf(TEXT("Expected '%s' not to be {%s} for context '%s'"), *What, *NotExpected, ContextObject ? *ContextObject->GetName() : TEXT("")));
+	}
+	else
+	{
+		LogAssertPass(FString::Printf(TEXT("String assertion passed (%s)"), *What));
+	}
+}
+
+void AFunctionalTest::LogAssertFail(const FString& Message)
+{
+	FString FullMessage(Message);
+	if ( IsInStep() )
+	{
+		FullMessage.Append(TEXT(" in step: "));
+		FString StepName = TEXT("");
+		if ( StepName.IsEmpty() )
+		{
+			StepName = TEXT("<UN-NAMED STEP>");
+		}
+		FullMessage.Append(StepName);
+	}
+	UE_VLOG(this, LogFunctionalTest, Error, TEXT("%s"), *FullMessage);
+	UE_LOG(LogFunctionalTest, Error, TEXT("%s"), *FullMessage);
+}
+
+void AFunctionalTest::LogAssertPass(const FString& Message)
+{
+	FString FullMessage(Message);
+	if ( IsInStep() )
+	{
+		FullMessage.Append(TEXT(" in step: "));
+		FString StepName = TEXT("");
+		if ( StepName.IsEmpty() )
+		{
+			StepName = TEXT("<UN-NAMED STEP>");
+		}
+		FullMessage.Append(StepName);
+	}
+	UE_VLOG(this, LogFunctionalTest, Log, TEXT("%s"), *FullMessage);
+	UE_LOG(LogFunctionalTest, Log, TEXT("%s"), *FullMessage);
+}
 
 //////////////////////////////////////////////////////////////////////////
 
diff --git a/Engine/Source/Developer/FunctionalTesting/Private/FunctionalTestGameMode.cpp b/Engine/Source/Developer/FunctionalTesting/Private/FunctionalTestGameMode.cpp
new file mode 100644
--- /dev/null
+++ b/Engine/Source/Developer/FunctionalTesting/Private/FunctionalTestGameMode.cpp
@@ -0,0 +1,11 @@
+// Copyright 1998-2016 Epic Games, Inc. All Rights Reserved.
+
+#include "FunctionalTestingPrivatePCH.h"
+#include "FunctionalTestGameMode.h"
+
+AFunctionalTestGameMode::AFunctionalTestGameMode(const FObjectInitializer& ObjectInitializer)
+	: Super(ObjectInitializer)
+{
+	SpectatorClass = ASpectatorPawn::StaticClass();
+	DefaultPawnClass = ASpectatorPawn::StaticClass();
+}
diff --git a/Engine/Source/Developer/FunctionalTesting/Private/FunctionalTestLevelScript.cpp b/Engine/Source/Developer/FunctionalTesting/Private/FunctionalTestLevelScript.cpp
new file mode 100644
--- /dev/null
+++ b/Engine/Source/Developer/FunctionalTesting/Private/FunctionalTestLevelScript.cpp
@@ -0,0 +1,9 @@
+// Copyright 1998-2016 Epic Games, Inc. All Rights Reserved.
+
+#include "FunctionalTestingPrivatePCH.h"
+#include "FunctionalTestLevelScript.h"
+
+AFunctionalTestLevelScript::AFunctionalTestLevelScript( const FObjectInitializer& ObjectInitializer )
+	: Super(ObjectInitializer)
+{
+}
diff --git a/Engine/Source/Developer/FunctionalTesting/Private/FunctionalTestingManager.cpp b/Engine/Source/Developer/FunctionalTesting/Private/FunctionalTestingManager.cpp
--- a/Engine/Source/Developer/FunctionalTesting/Private/FunctionalTestingManager.cpp
+++ b/Engine/Source/Developer/FunctionalTesting/Private/FunctionalTestingManager.cpp
@@ -117,7 +117,7 @@ bool UFunctionalTestingManager::RunAllFunctionalTests(UObject* WorldContext, boo
 		for (TActorIterator<AFunctionalTest> It(WorldContext->GetWorld()); It; ++It)
 		{
 			AFunctionalTest* Test = (*It);
-			if (Test != NULL && Test->bIsEnabled == true)
+			if (Test != nullptr && Test->IsEnabled() == true)
 			{
 				Manager->AllTests.Add(Test);
 			}
@@ -147,7 +147,7 @@ bool UFunctionalTestingManager::RunAllFunctionalTests(UObject* WorldContext, boo
 void UFunctionalTestingManager::TriggerFirstValidTest()
 {
 	UWorld* World = GetWorld();
-	bIsRunning = World != NULL && World->GetNavigationSystem() != NULL;
+	bIsRunning = World != nullptr && World->GetNavigationSystem() != nullptr;
 
 	if (bInitialDelayApplied == true && (bWaitForNavigationBuildFinish == false || UNavigationSystem::IsNavigationBeingBuilt(World) == false) && World->AreActorsInitialized())
 	{
@@ -211,14 +211,14 @@ void UFunctionalTestingManager::OnTestDone(AFunctionalTest* FTest)
 
 void UFunctionalTestingManager::NotifyTestDone(AFunctionalTest* FTest)
 {
-	if (FTest->IsSuccessful() == false)
-	{
-		if (GatheredFailedTestsReproString.IsEmpty() == false)
-		{
-			GatheredFailedTestsReproString += FFunctionalTesting::ReproStringTestSeparator;
-		}
-		GatheredFailedTestsReproString += FTest->GetReproString();
-	}
+	//if (FTest->IsSuccessful() == false)
+	//{
+	//	if (GatheredFailedTestsReproString.IsEmpty() == false)
+	//	{
+	//		GatheredFailedTestsReproString += FFunctionalTesting::ReproStringTestSeparator;
+	//	}
+	//	GatheredFailedTestsReproString += FTest->GetReproString();
+	//}
 
 	if (FTest->OnWantsReRunCheck() == false && FTest->WantsToRunAgain() == false)
 	{
@@ -256,10 +256,13 @@ void UFunctionalTestingManager::NotifyTestDone(AFunctionalTest* FTest)
 
 void UFunctionalTestingManager::AllTestsDone()
 {
-	if (GatheredFailedTestsReproString.IsEmpty() == false)
-	{
-		UE_LOG(LogFunctionalTest, Log, TEXT("Repro String : %s"), *GatheredFailedTestsReproString);
-	}
+	//TODO AUTOMATION Should we revive this?  There's no good way at the moment to know if the test -actually-
+	// failed, because warnings and such could actually fail the test.
+
+	//if (GatheredFailedTestsReproString.IsEmpty() == false)
+	//{
+	//	UE_LOG(LogFunctionalTest, Log, TEXT("Repro String : %s"), *GatheredFailedTestsReproString);
+	//}
 
 	if (bLooped == true)
 	{
@@ -268,7 +271,7 @@ void UFunctionalTestingManager::AllTestsDone()
 		// reset
 		ensure(TestReproStrings.Num() == 0);
 		SetReproString(StartingReproString);
-		GatheredFailedTestsReproString = TEXT("");
+		//GatheredFailedTestsReproString = TEXT("");
 		TestsLeft = AllTests;
 
 		UE_LOG(LogFunctionalTest, Log, TEXT("----- Starting iteration %d -----"), CurrentIteration);
@@ -299,7 +302,9 @@ bool UFunctionalTestingManager::RunFirstValidTest()
 		while (TestReproStrings.Num() > 0)
 		{
 			TArray<FString> TestParams;
-			const FString SingleTestReproString = TestReproStrings.Pop(/*bAllowShrinking=*/false);
+			const FString SingleTestReproString = TestReproStrings[0];
+			TestReproStrings.RemoveAt(0);
+
 			SingleTestReproString.ParseIntoArray(TestParams, TEXT("#"), /*InCullEmpty=*/true);
 			
 			if (TestParams.Num() == 0)
@@ -307,6 +312,7 @@ bool UFunctionalTestingManager::RunFirstValidTest()
 				UE_LOG(LogFunctionalTest, Warning, TEXT("Unable to parse \'%s\'"), *SingleTestReproString);
 				continue;
 			}
+
 			// first param is the test name. Look for it		
 			const FString TestName = TestParams[0];
 			TestParams.RemoveAt(0, 1, /*bAllowShrinking=*/false);
@@ -314,7 +320,7 @@ bool UFunctionalTestingManager::RunFirstValidTest()
 			if (TestToRun)
 			{
 				TestToRun->TestFinishedObserver = TestFinishedObserver;
-				if (TestToRun->StartTest(TestParams))
+				if (TestToRun->RunTest(TestParams))
 				{
 					bTestSuccessfullyTriggered = true;
 					break;
@@ -338,9 +344,9 @@ bool UFunctionalTestingManager::RunFirstValidTest()
 			bool bRemove = TestsLeft[Index] == NULL;
 			if (TestsLeft[Index] != NULL)
 			{
-				ensure(TestsLeft[Index]->bIsEnabled);
+				ensure(TestsLeft[Index]->IsEnabled());
 				TestsLeft[Index]->TestFinishedObserver = TestFinishedObserver;
-				if (TestsLeft[Index]->StartTest())
+				if (TestsLeft[Index]->RunTest())
 				{
 					if (TestsLeft[Index]->IsRunning() == true)
 					{
diff --git a/Engine/Source/Developer/FunctionalTesting/Private/FunctionalTestingModule.cpp b/Engine/Source/Developer/FunctionalTesting/Private/FunctionalTestingModule.cpp
--- a/Engine/Source/Developer/FunctionalTesting/Private/FunctionalTestingModule.cpp
+++ b/Engine/Source/Developer/FunctionalTesting/Private/FunctionalTestingModule.cpp
@@ -1,7 +1,10 @@
 // Copyright 1998-2016 Epic Games, Inc. All Rights Reserved.
 
 #include "FunctionalTestingPrivatePCH.h"
-#include "ModuleInterface.h"
+#include "Modules/ModuleInterface.h"
+#include "Modules/ModuleManager.h"
+
+#define LOCTEXT_NAMESPACE "FunctionalTesting"
 
 void FFunctionalTestingModule::StartupModule() 
 {
@@ -10,7 +13,9 @@ void FFunctionalTestingModule::StartupModule()
 
 void FFunctionalTestingModule::ShutdownModule() 
 {
-	Manager = NULL;
+	Manager.Reset();
 }
 
 IMPLEMENT_MODULE( FFunctionalTestingModule, FunctionalTesting );
+
+#undef LOCTEXT_NAMESPACE
diff --git a/Engine/Source/Developer/FunctionalTesting/Private/Tests/ClientFuncTestPerforming.cpp b/Engine/Source/Developer/FunctionalTesting/Private/Tests/ClientFuncTestPerforming.cpp
--- a/Engine/Source/Developer/FunctionalTesting/Private/Tests/ClientFuncTestPerforming.cpp
+++ b/Engine/Source/Developer/FunctionalTesting/Private/Tests/ClientFuncTestPerforming.cpp
@@ -2,6 +2,7 @@
 
 #include "FunctionalTestingPrivatePCH.h"
 #include "AutomationCommon.h"
+#include "AssetRegistryModule.h"
 
 #if WITH_DEV_AUTOMATION_TESTS
 
@@ -36,14 +37,47 @@ bool FTriggerFTests::Update()
 }
 
 
+DEFINE_LATENT_AUTOMATION_COMMAND_ONE_PARAMETER(FTriggerFTest, FString, TestName);
+bool FTriggerFTest::Update()
+{
+	IFuncTestManager* Manager = FFunctionalTestingModule::Get();
+	if ( Manager->IsFinished() )
+	{
+		// if tests have been already triggered by level script just make sure it's not looping
+		if ( Manager->IsRunning() )
+		{
+			FFunctionalTestingModule::Get()->SetLooping(false);
+		}
+		else
+		{
+			FFunctionalTestingModule::Get()->RunTestOnMap(TestName, false, false);
+		}
+	}
+
+	return true;
+}
+
+
 DEFINE_LATENT_AUTOMATION_COMMAND(FStartFTestsOnMap);
 bool FStartFTestsOnMap::Update()
 {
+	// This should now be handled by your IsReady override of the functional test.
 	//should really be wait until the map is properly loaded....in PIE or gameplay....
-	ADD_LATENT_AUTOMATION_COMMAND(FWaitLatentCommand(15.f));
+	//ADD_LATENT_AUTOMATION_COMMAND(FWaitLatentCommand(15.f));
+
 	ADD_LATENT_AUTOMATION_COMMAND(FTriggerFTests);
 	ADD_LATENT_AUTOMATION_COMMAND(FWaitForFTestsToFinish);	
-	
+
+	return true;
+}
+
+
+DEFINE_LATENT_AUTOMATION_COMMAND_ONE_PARAMETER(FStartFTestOnMap, FString, TestName);
+bool FStartFTestOnMap::Update()
+{
+	ADD_LATENT_AUTOMATION_COMMAND(FTriggerFTest(TestName));
+	ADD_LATENT_AUTOMATION_COMMAND(FWaitForFTestsToFinish);
+
 	return true;
 }
 
@@ -51,35 +85,84 @@ bool FStartFTestsOnMap::Update()
  * 
  */
 
+// Project.Maps.Client Functional Testing
+// Project.Maps.Functional Tests
+
+
+class FClientFunctionalTestingMapsBase : public FAutomationTestBase
+{
+public:
+	FClientFunctionalTestingMapsBase(const FString& InName, const bool bInComplexTask)
+		: FAutomationTestBase(InName, bInComplexTask)
+	{
+	}
+
+	virtual FString GetTestOpenCommand(const FString& Parameter) const override
+	{
+		FString MapPath;
+		FString MapTestName;
+		Parameter.Split(TEXT(";"), &MapPath, &MapTestName);
+
+		return FString::Printf(TEXT("Automate.OpenMapAndFocusActor %s %s"), *MapPath, *MapTestName);
+	}
+
+	virtual FString GetTestAssetPath(const FString& Parameter) const override
+	{
+		FString MapPath;
+		FString MapTestName;
+		Parameter.Split(TEXT(";"), &MapPath, &MapTestName);
+
+		return MapPath;
+	}
+};
+
+
 // create test base class
-IMPLEMENT_COMPLEX_AUTOMATION_TEST(FClientFunctionalTestingMapsBase, "Project.Maps.Client Functional Testing", (EAutomationTestFlags::ClientContext | EAutomationTestFlags::EditorContext | EAutomationTestFlags::ProductFilter))
+IMPLEMENT_CUSTOM_COMPLEX_AUTOMATION_TEST(FClientFunctionalTestingMaps, FClientFunctionalTestingMapsBase, "Project.Functional Tests", (EAutomationTestFlags::ClientContext | EAutomationTestFlags::EditorContext | EAutomationTestFlags::ProductFilter))
 
 
 /** 
  * Requests a enumeration of all maps to be loaded
  */
-void FClientFunctionalTestingMapsBase::GetTests(TArray<FString>& OutBeautifiedNames, TArray <FString>& OutTestCommands) const
+void FClientFunctionalTestingMaps::GetTests(TArray<FString>& OutBeautifiedNames, TArray<FString>& OutTestCommands) const
 {
-	TArray<FString> FileList;
-	// Look directly on disk. Very slow!
-	FPackageName::FindPackagesInDirectory(FileList, *FPaths::GameContentDir());
+	IAssetRegistry& AssetRegistry = FModuleManager::Get().LoadModuleChecked<FAssetRegistryModule>("AssetRegistry").Get();
 
-	// Iterate over all files, adding the ones with the map extension..
-	for (int32 FileIndex = 0; FileIndex < FileList.Num(); FileIndex++)
+	if ( !AssetRegistry.IsLoadingAssets() )
 	{
-		const FString& Filename = FileList[FileIndex];
-
-		// Disregard filenames that don't have the map extension if we're in MAPSONLY mode.
-		if (FPaths::GetExtension(Filename, true) == FPackageName::GetMapPackageExtension())
+		TArray<FAssetData> MapList;
+		if ( AssetRegistry.GetAssetsByClass(UWorld::StaticClass()->GetFName(), /*out*/ MapList) )
 		{
-			if (FAutomationTestFramework::GetInstance().ShouldTestContent(Filename))
+			for ( const FAssetData& MapAsset : MapList )
 			{
-				const FString BaseFilename = FPaths::GetBaseFilename(Filename);
-				// Disregard filenames that don't have the map extension if we're in MAPSONLY mode.
-				if (BaseFilename.Find(TEXT("FTEST_")) == 0)
+				const FString* Tests = MapAsset.TagsAndValues.Find(TEXT("Tests"));
+				const FString* TestNames = MapAsset.TagsAndValues.Find(TEXT("TestNames"));
+
+				if ( Tests && TestNames )
 				{
-					OutBeautifiedNames.Add(BaseFilename);
-					OutTestCommands.Add(Filename);
+					int32 TestCount = FCString::Atoi(**Tests);
+					if ( TestCount > 0 )
+					{
+						TArray<FString> MapTests;
+						( *TestNames ).ParseIntoArray(MapTests, TEXT(";"), true);
+
+						for ( const FString& MapTest : MapTests )
+						{
+							FString BeautifulTestName;
+							FString RealTestName;
+
+							if ( MapTest.Split(TEXT("|"), &BeautifulTestName, &RealTestName) )
+							{
+								OutBeautifiedNames.Add(MapAsset.AssetName.ToString() + TEXT(".") + *BeautifulTestName);
+								OutTestCommands.Add(MapAsset.PackageName.ToString() + TEXT(";") + *RealTestName);
+							}
+						}
+					}
+				}
+				else if ( MapAsset.AssetName.ToString().Find(TEXT("FTEST_")) == 0 )
+				{
+					OutBeautifiedNames.Add(MapAsset.AssetName.ToString());
+					OutTestCommands.Add(MapAsset.PackageName.ToString());
 				}
 			}
 		}
@@ -87,24 +170,71 @@ void FClientFunctionalTestingMapsBase::GetTests(TArray<FString>& OutBeautifiedNa
 }
 
 
+// @todo this is a temporary solution. Once we know how to get test's hands on a proper world
+// this function should be redone/removed
+static UWorld* GetAnyGameWorld()
+{
+	UWorld* TestWorld = nullptr;
+	const TIndirectArray<FWorldContext>& WorldContexts = GEngine->GetWorldContexts();
+	for ( const FWorldContext& Context : WorldContexts )
+	{
+		if ( ( ( Context.WorldType == EWorldType::PIE ) || ( Context.WorldType == EWorldType::Game ) ) && ( Context.World() != NULL ) )
+		{
+			TestWorld = Context.World();
+			break;
+		}
+	}
+
+	return TestWorld;
+}
+
+
 /** 
  * Execute the loading of each map and performance captures
  *
  * @param Parameters - Should specify which map name to load
  * @return	TRUE if the test was successful, FALSE otherwise
  */
-bool FClientFunctionalTestingMapsBase::RunTest(const FString& Parameters)
+bool FClientFunctionalTestingMaps::RunTest(const FString& Parameters)
 {
-	FString MapName = Parameters;
+	TArray<FString> ParamArray;
+	Parameters.ParseIntoArray(ParamArray, TEXT(";"), true);
+
+	FString MapName = ParamArray[0];
+	FString MapTestName = ( ParamArray.Num() > 1 ) ? ParamArray[1] : TEXT("");
+
+	bool bCanProceed = false;
 
-	bool bCanProceed = AutomationOpenMap(MapName);
+	UWorld* TestWorld = GetAnyGameWorld();
+	if ( TestWorld && TestWorld->GetMapName() == MapName )
+	{
+		// Map is already loaded.
+		bCanProceed = true;
+	}
+	else
+	{
+		bCanProceed = AutomationOpenMap(MapName);
+	}
 
 	if (bCanProceed)
 	{
-		ADD_LATENT_AUTOMATION_COMMAND(FStartFTestsOnMap());
+		if ( MapTestName.IsEmpty() )
+		{
+			ADD_LATENT_AUTOMATION_COMMAND(FStartFTestsOnMap());
+		}
+		else
+		{
+			ADD_LATENT_AUTOMATION_COMMAND(FStartFTestOnMap(MapTestName));
+		}
+
 		return true;
 	}
 
+	/// FAutomationTestFramework::GetInstance().UnregisterAutomationTest
+
+	//	ADD_LATENT_AUTOMATION_COMMAND(FWaitLatentCommand(1.f));
+	//  ADD_LATENT_AUTOMATION_COMMAND(FExitGameCommand);
+
 	UE_LOG(LogFunctionalTesting, Error, TEXT("Failed to start the %s map (possibly due to BP compilation issues)"), *MapName);
 	return false;
 }
diff --git a/Engine/Source/Developer/FunctionalTesting/Public/AutomationBlueprintFunctionLibrary.h b/Engine/Source/Developer/FunctionalTesting/Public/AutomationBlueprintFunctionLibrary.h
new file mode 100644
--- /dev/null
+++ b/Engine/Source/Developer/FunctionalTesting/Public/AutomationBlueprintFunctionLibrary.h
@@ -0,0 +1,28 @@
+// Copyright 1998-2016 Epic Games, Inc. All Rights Reserved.
+#pragma once
+
+#include "Core.h"
+#include "Kismet/BlueprintFunctionLibrary.h"
+
+#include "AutomationBlueprintFunctionLibrary.generated.h"
+
+UCLASS()
+class FUNCTIONALTESTING_API UAutomationBlueprintFunctionLibrary : public UBlueprintFunctionLibrary
+{
+	GENERATED_UCLASS_BODY()
+	
+public:
+	static void TakeAutomationScreenshotInternal(const FString& Name);
+
+	UFUNCTION(BlueprintCallable, Category = "StudiosAutomation", meta = (Latent, HidePin = "WorldContextObject", DefaultToSelf = "WorldContextObject", LatentInfo = "LatentInfo", Name = ""))
+	static void TakeAutomationScreenshot(UObject* WorldContextObject, FLatentActionInfo LatentInfo, const FString& Name);
+
+	UFUNCTION(BlueprintCallable, Category = "StudiosAutomation|Helpers", meta = (Latent, HidePin = "WorldContextObject", DefaultToSelf = "WorldContextObject", LatentInfo = "LatentInfo", NameOverride = "", DelayBeforeScreenshotSeconds = "0.2"))
+	static void TakeAutomationScreenshotAtCamera(UObject* WorldContextObject, FLatentActionInfo LatentInfo, ACameraActor* Camera, const FString& NameOverride, float DelayBeforeScreenshotSeconds);
+
+    UFUNCTION(BlueprintCallable, Category = "StudiosAutomation|Performance")
+    static void BeginPerformanceCapture();
+
+    UFUNCTION(BlueprintCallable, Category = "StudiosAutomation|Performance")
+    static void EndPerformanceCapture();
+};
\ No newline at end of file
diff --git a/Engine/Source/Developer/FunctionalTesting/Public/FunctionalTestingModule.h b/Engine/Source/Developer/FunctionalTesting/Public/FunctionalTestingModule.h
--- a/Engine/Source/Developer/FunctionalTesting/Public/FunctionalTestingModule.h
+++ b/Engine/Source/Developer/FunctionalTesting/Public/FunctionalTestingModule.h
@@ -17,7 +17,10 @@ public:
 		FFunctionalTestingModule& Module = FModuleManager::Get().LoadModuleChecked<FFunctionalTestingModule>("FunctionalTesting");
 		return Module.GetSingleton();
 	}
+
 private:
+	void OnModulesChanged(FName Module, EModuleChangeReason Reason);
+
 	class IFuncTestManager* GetSingleton() const 
 	{ 
 		return Manager.Get(); 
diff --git a/Engine/Source/Developer/FunctionalTesting/Public/IFuncTestManager.h b/Engine/Source/Developer/FunctionalTesting/Public/IFuncTestManager.h
--- a/Engine/Source/Developer/FunctionalTesting/Public/IFuncTestManager.h
+++ b/Engine/Source/Developer/FunctionalTesting/Public/IFuncTestManager.h
@@ -8,6 +8,8 @@ public:
 	/** Triggers in sequence all functional tests found on the level.*/
 	virtual void RunAllTestsOnMap(bool bClearLog, bool bRunLooped) = 0;
 
+	virtual void RunTestOnMap(const FString& TestName, bool bClearLog, bool bRunLooped) = 0;
+
 	virtual bool IsRunning() const = 0;
 
 	virtual bool IsFinished() const = 0;
diff --git a/Engine/Source/Developer/LogVisualizer/Private/LogVisualizerModule.cpp b/Engine/Source/Developer/LogVisualizer/Private/LogVisualizerModule.cpp
--- a/Engine/Source/Developer/LogVisualizer/Private/LogVisualizerModule.cpp
+++ b/Engine/Source/Developer/LogVisualizer/Private/LogVisualizerModule.cpp
@@ -51,7 +51,7 @@ void FLogVisualizerModule::StartupModule()
 	ISettingsModule* SettingsModule = FModuleManager::GetModulePtr<ISettingsModule>("Settings");
 	if (SettingsModule != nullptr)
 	{
-		SettingsModule->RegisterSettings("Editor", "General", "VisualLogger",
+		SettingsModule->RegisterSettings("Editor", "Advanced", "VisualLogger",
 			LOCTEXT("AIToolsSettingsName", "Visual Logger"),
 			LOCTEXT("AIToolsSettingsDescription", "General settings for UE4 AI Tools."),
 			ULogVisualizerSettings::StaticClass()->GetDefaultObject()
diff --git a/Engine/Source/Developer/LogVisualizer/Private/LogVisualizerStyle.cpp b/Engine/Source/Developer/LogVisualizer/Private/LogVisualizerStyle.cpp
--- a/Engine/Source/Developer/LogVisualizer/Private/LogVisualizerStyle.cpp
+++ b/Engine/Source/Developer/LogVisualizer/Private/LogVisualizerStyle.cpp
@@ -116,7 +116,7 @@ TSharedRef< FSlateStyleSet > FLogVisualizerStyle::Create()
 	// Toolbar
 	{
 		Style.Set("Toolbar.Pause", new IMAGE_BRUSH("Icons/icon_pause_40x", Icon40x40));
-		Style.Set("Toolbar.Resume", new IMAGE_BRUSH("Icons/icon_play_40x", Icon40x40));
+		Style.Set("Toolbar.Resume", new IMAGE_BRUSH("Icons/icon_simulate_40x", Icon40x40));
 		Style.Set("Toolbar.Record", new IMAGE_BRUSH("Icons/LV_Record", Icon40x40));
 		Style.Set("Toolbar.Stop", new IMAGE_BRUSH("Icons/LV_Stop", Icon40x40));
 		Style.Set("Toolbar.Camera", new IMAGE_BRUSH("Icons/LV_Camera", Icon40x40));
@@ -128,7 +128,7 @@ TSharedRef< FSlateStyleSet > FLogVisualizerStyle::Create()
 		Style.Set("Toolbar.HideEmpty", new IMAGE_BRUSH("Icons/icon_levels_visible_40x", Icon40x40));
 
 		Style.Set("Toolbar.Pause.small", new IMAGE_BRUSH("Icons/icon_pause_40x", Icon20x20));
-		Style.Set("Toolbar.Resume.small", new IMAGE_BRUSH("Icons/icon_play_40x", Icon20x20));
+		Style.Set("Toolbar.Resume.small", new IMAGE_BRUSH("Icons/icon_simulate_40x", Icon20x20));
 		Style.Set("Toolbar.Record.small", new IMAGE_BRUSH("Icons/LV_Record", Icon20x20));
 		Style.Set("Toolbar.Stop.small", new IMAGE_BRUSH("Icons/LV_Stop", Icon20x20));
 		Style.Set("Toolbar.Camera.small", new IMAGE_BRUSH("Icons/LV_Camera", Icon20x20));
diff --git a/Engine/Source/Developer/LogVisualizer/Private/SVisualLoggerTimelinesContainer.cpp b/Engine/Source/Developer/LogVisualizer/Private/SVisualLoggerTimelinesContainer.cpp
--- a/Engine/Source/Developer/LogVisualizer/Private/SVisualLoggerTimelinesContainer.cpp
+++ b/Engine/Source/Developer/LogVisualizer/Private/SVisualLoggerTimelinesContainer.cpp
@@ -14,8 +14,8 @@ TSharedRef<SWidget> SVisualLoggerTimelinesContainer::GetRightClickMenuContent()
 	MenuBuilder.BeginSection("VisualLogReports", LOCTEXT("VisualLogReports", "VisualLog Reports"));
 	{
 		MenuBuilder.AddMenuEntry(
-			LOCTEXT("GenarateReport", "Genarate Report"),
-			LOCTEXT("GenarateReportTooltip", "Genarate report from Visual Log events."),
+			LOCTEXT("GenerateReport", "Generate  Report"),
+			LOCTEXT("GenerateReportTooltip", "Generate report from Visual Log events."),
 			FSlateIcon(),
 			FUIAction(FExecuteAction::CreateSP(this, &SVisualLoggerTimelinesContainer::GenerateReport))
 			);
diff --git a/Engine/Source/Developer/LogVisualizer/Private/VisualLoggerCommands.h b/Engine/Source/Developer/LogVisualizer/Private/VisualLoggerCommands.h
--- a/Engine/Source/Developer/LogVisualizer/Private/VisualLoggerCommands.h
+++ b/Engine/Source/Developer/LogVisualizer/Private/VisualLoggerCommands.h
@@ -23,7 +23,7 @@ public:
 		UI_COMMAND(StartRecording, "Start", "Start the debugger", EUserInterfaceActionType::Button, FInputChord());
 		UI_COMMAND(StopRecording, "Stop", "Step over the current message", EUserInterfaceActionType::Button, FInputChord());
 		UI_COMMAND(Pause, "Pause", "Stop the debugger", EUserInterfaceActionType::Button, FInputChord());
-		UI_COMMAND(Resume, "Resume", "Stop the debugger", EUserInterfaceActionType::Button, FInputChord());
+		UI_COMMAND(Resume, "Resume", "Resume", EUserInterfaceActionType::Button, FInputChord());
 		UI_COMMAND(FreeCamera, "FreeCamera", "Enable free camera", EUserInterfaceActionType::ToggleButton, FInputChord());
 		UI_COMMAND(LoadFromVLog, "Load", "Load external vlogs", EUserInterfaceActionType::Button, FInputChord());
 		UI_COMMAND(SaveToVLog, "Save", "Save selected data to vlog file", EUserInterfaceActionType::Button, FInputChord());
diff --git a/Engine/Source/Developer/MaterialUtilities/Private/MaterialUtilities.cpp b/Engine/Source/Developer/MaterialUtilities/Private/MaterialUtilities.cpp
--- a/Engine/Source/Developer/MaterialUtilities/Private/MaterialUtilities.cpp
+++ b/Engine/Source/Developer/MaterialUtilities/Private/MaterialUtilities.cpp
@@ -1254,7 +1254,7 @@ UMaterialInstanceConstant* FMaterialUtilities::CreateInstancedMaterial(UMaterial
 	if (ExistingPackage && !ExistingPackage->IsA<UMaterialInstanceConstant>())
 	{
 #if WITH_EDITOR
-		AutomationEditorCommonUtils::NullReferencesToObject(ExistingPackage);		
+		FAutomationEditorCommonUtils::NullReferencesToObject(ExistingPackage);		
 #endif // WITH_EDITOR
 		ExistingPackage->MarkPendingKill();
 		CollectGarbage(GARBAGE_COLLECTION_KEEPFLAGS, true);
diff --git a/Engine/Source/Developer/MessageLog/Private/UserInterface/SMessageLogMessageListRow.cpp b/Engine/Source/Developer/MessageLog/Private/UserInterface/SMessageLogMessageListRow.cpp
--- a/Engine/Source/Developer/MessageLog/Private/UserInterface/SMessageLogMessageListRow.cpp
+++ b/Engine/Source/Developer/MessageLog/Private/UserInterface/SMessageLogMessageListRow.cpp
@@ -3,6 +3,7 @@
 #include "MessageLogPrivatePCH.h"
 #include "SMessageLogMessageListRow.h"
 #include "SHyperlink.h"
+#include "Internationalization/Regex.h"
 
 #define LOCTEXT_NAMESPACE "SMessageLogMessageListRow"
 
@@ -229,9 +230,47 @@ void SMessageLogMessageListRow::CreateMessage(const TSharedRef<SHorizontalBox>&
 		}
 		else
 		{
-			RowContent = SNew(STextBlock)
+			FString MessageString = InMessageToken->ToText().ToString();
+
+			// ^((?:[\w]\:|\\)(?:(?:\\[a-z_\-\s0-9\.]+)+)\.(?:cpp|h))\((\d+)\)
+			// https://regex101.com/r/vV4cV7/1
+			FRegexPattern FileAndLinePattern(TEXT("^((?:[\\w]\\:|\\\\)(?:(?:\\\\[a-z_\\-\\s0-9\\.]+)+)\\.(?:cpp|h))\\((\\d+)\\)"));
+			FRegexMatcher FileAndLineRegexMatcher(FileAndLinePattern, MessageString);
+
+			TSharedRef<SWidget> SourceLink = SNullWidget::NullWidget;
+
+			if ( FileAndLineRegexMatcher.FindNext() )
+			{
+				FString FileName = FileAndLineRegexMatcher.GetCaptureGroup(1);
+				int32 LineNumber = FCString::Atoi(*FileAndLineRegexMatcher.GetCaptureGroup(2));
+
+				// Remove the hyperlink from the message, since we're splitting it into its own string.
+				MessageString = MessageString.RightChop(FileAndLineRegexMatcher.GetMatchEnding());
+
+				SourceLink = SNew(SHyperlink)
+					.Style(FEditorStyle::Get(), "Common.GotoNativeCodeHyperlink")
+					.TextStyle(FEditorStyle::Get(), "MessageLog")
+					.OnNavigate_Lambda([=] { FSlateApplication::Get().GotoLineInSource(FileName, LineNumber); })
+					.Text(FText::FromString(FileAndLineRegexMatcher.GetCaptureGroup(0)));
+			}
+
+			RowContent = SNew(SHorizontalBox)
+
+			+ SHorizontalBox::Slot()
+			.AutoWidth()
+			.Padding(0)
+			[
+				SourceLink
+			]
+
+			+ SHorizontalBox::Slot()
+			.AutoWidth()
+			.Padding(0)
+			[
+				SNew(STextBlock)
 				.ColorAndOpacity(FSlateColor::UseSubduedForeground())
-				.Text(InMessageToken->ToText());
+				.Text(FText::FromString(MessageString))
+			];
 		}
 	}
 		break;
diff --git a/Engine/Source/Developer/ModuleUI/Private/SModuleUI.cpp b/Engine/Source/Developer/ModuleUI/Private/SModuleUI.cpp
--- a/Engine/Source/Developer/ModuleUI/Private/SModuleUI.cpp
+++ b/Engine/Source/Developer/ModuleUI/Private/SModuleUI.cpp
@@ -2,6 +2,7 @@
 
 #include "ModuleUIPrivatePCH.h"
 #include "HotReloadInterface.h"
+#include "SSearchBox.h"
 
 BEGIN_SLATE_FUNCTION_BUILD_OPTIMIZATION
 void SModuleUI::Construct(const SModuleUI::FArguments& InArgs)
@@ -13,6 +14,13 @@ void SModuleUI::Construct(const SModuleUI::FArguments& InArgs)
 	[
 		SNew( SVerticalBox )
 
+			+ SVerticalBox::Slot()
+				.AutoHeight()
+			[
+				SAssignNew(ModuleNameSearchBox, SSearchBox)
+				.OnTextChanged(this, &SModuleUI::OnFilterTextChanged)
+			]
+
 			// List of modules
 			+ SVerticalBox::Slot()
 				.FillHeight( 1.0f )		// We want the list to stretch vertically to fill up the user-resizable space
@@ -53,7 +61,10 @@ SModuleUI::~SModuleUI()
 	FModuleManager::Get().OnModulesChanged().RemoveAll( this );
 }
 
-
+void SModuleUI::OnFilterTextChanged(const FText& InFilterText)
+{
+	UpdateModuleListItems();
+}
 TSharedRef<ITableRow> SModuleUI::OnGenerateWidgetForModuleListView(TSharedPtr< FModuleListItem > InItem, const TSharedRef<STableViewBase>& OwnerTable)
 {
 	#define LOCTEXT_NAMESPACE "ModuleUI"
@@ -163,9 +174,17 @@ void SModuleUI::UpdateModuleListItems()
 
 		FName ModuleName(*ModuleStatus.Name);
 
-		TSharedRef< FModuleListItem > NewItem( new FModuleListItem() );
-		NewItem->ModuleName = ModuleName;
-		ModuleListItems.Add( NewItem );
+		FString ModuleNameStr = ModuleName.ToString();
+		FString FilterStr = ModuleNameSearchBox->GetText().ToString();
+		FilterStr.Trim();
+		FilterStr.TrimTrailing();
+		if ( (!FilterStr.IsEmpty() && ModuleNameStr.Contains(FilterStr)) || 
+			  FilterStr.IsEmpty() )
+		{
+			TSharedRef< FModuleListItem > NewItem(new FModuleListItem());
+			NewItem->ModuleName = ModuleName;
+			ModuleListItems.Add(NewItem);
+		}
 	}
 
 
diff --git a/Engine/Source/Developer/Profiler/Private/ProfilerManager.cpp b/Engine/Source/Developer/Profiler/Private/ProfilerManager.cpp
--- a/Engine/Source/Developer/Profiler/Private/ProfilerManager.cpp
+++ b/Engine/Source/Developer/Profiler/Private/ProfilerManager.cpp
@@ -146,11 +146,6 @@ FProfilerManager::~FProfilerManager()
 	// Remove ourselves from the session manager.
 	if (SessionManager.IsValid())
 	{
-		// clear the selected session
-		ISessionInfoPtr Ptr;
-		Ptr.Reset();
-		SessionManager->SelectSession( Ptr );
-
 		SessionManager->OnCanSelectSession().RemoveAll( this );
 		SessionManager->OnSelectedSessionChanged().RemoveAll( this );
 		SessionManager->OnInstanceSelectionChanged().RemoveAll( this );
diff --git a/Engine/Source/Developer/SessionFrontend/Private/SessionFrontendModule.cpp b/Engine/Source/Developer/SessionFrontend/Private/SessionFrontendModule.cpp
--- a/Engine/Source/Developer/SessionFrontend/Private/SessionFrontendModule.cpp
+++ b/Engine/Source/Developer/SessionFrontend/Private/SessionFrontendModule.cpp
@@ -28,6 +28,18 @@ public:
 	{
 		return SNew(SSessionConsole, SessionManager);
 	}
+
+	virtual void InvokeSessionFrontend(FName SubTabToActivate = NAME_None) override
+	{
+		FGlobalTabmanager::Get()->InvokeTab(SessionFrontendTabName);
+		if ( WeakFrontend.IsValid() )
+		{
+			if ( SubTabToActivate != NAME_None )
+			{
+				WeakFrontend.Pin()->GetTabManager()->InvokeTab(SubTabToActivate);
+			}
+		}
+	}
 	
 public:
 
@@ -65,10 +77,16 @@ private:
 		const TSharedRef<SDockTab> DockTab = SNew(SDockTab)
 			.TabRole(ETabRole::MajorTab);
 
-		DockTab->SetContent(SNew(SSessionFrontend, DockTab, SpawnTabArgs.GetOwnerWindow()));
+		TSharedRef<SSessionFrontend> Frontend = SNew(SSessionFrontend, DockTab, SpawnTabArgs.GetOwnerWindow());
+		WeakFrontend = Frontend;
+
+		DockTab->SetContent(Frontend);
 
 		return DockTab;
 	}
+
+private:
+	TWeakPtr<SSessionFrontend> WeakFrontend;
 };
 
 
diff --git a/Engine/Source/Developer/SessionFrontend/Private/Widgets/Browser/SSessionBrowser.cpp b/Engine/Source/Developer/SessionFrontend/Private/Widgets/Browser/SSessionBrowser.cpp
--- a/Engine/Source/Developer/SessionFrontend/Private/Widgets/Browser/SSessionBrowser.cpp
+++ b/Engine/Source/Developer/SessionFrontend/Private/Widgets/Browser/SSessionBrowser.cpp
@@ -31,7 +31,7 @@ BEGIN_SLATE_FUNCTION_BUILD_OPTIMIZATION
 void SSessionBrowser::Construct( const FArguments& InArgs, ISessionManagerRef InSessionManager )
 {
 	IgnoreSessionManagerEvents = false;
-	IgnoreSessionTreeEvents = false;
+	updatingTreeExpansion = false;
 	SessionManager = InSessionManager;
 
 	ChildSlot
@@ -39,85 +39,45 @@ void SSessionBrowser::Construct( const FArguments& InArgs, ISessionManagerRef In
 		SNew(SVerticalBox)
 
 		+ SVerticalBox::Slot()
-			.FillHeight(1.0f)
+		.FillHeight(1.0f)
+		[
+			// session tree
+			SNew(SBorder)
+			.BorderImage(FEditorStyle::GetBrush("ToolPanel.GroupBorder"))
+			.Padding(0.0f)
 			[
-				// session tree
-				SNew(SBorder)
-					.BorderImage(FEditorStyle::GetBrush("ToolPanel.GroupBorder"))
-					.Padding(0.0f)
-					[
-						SAssignNew(SessionTreeView, STreeView<TSharedPtr<FSessionBrowserTreeItem>>)
-							.ItemHeight(20.0f)
-							.OnExpansionChanged(this, &SSessionBrowser::HandleSessionTreeViewExpansionChanged)
-							.OnGenerateRow(this, &SSessionBrowser::HandleSessionTreeViewGenerateRow)
-							.OnGetChildren(this, &SSessionBrowser::HandleSessionTreeViewGetChildren)
-							.OnSelectionChanged(this, &SSessionBrowser::HandleSessionTreeViewSelectionChanged)
-							.SelectionMode(ESelectionMode::Multi)
-							.TreeItemsSource(&SessionTreeItems)
-							.HeaderRow
-							(
-								SNew(SHeaderRow)
-
-								+ SHeaderRow::Column("Name")
-									.DefaultLabel(LOCTEXT("InstanceListNameColumnHeader", "Name"))
-									.FillWidth(0.3f)
-
-								+ SHeaderRow::Column("Type")
-									.DefaultLabel(LOCTEXT("InstanceListTypeColumnHeader", "Type"))
-									.FillWidth(0.2f)
-
-								+ SHeaderRow::Column("Device")
-									.DefaultLabel(LOCTEXT("InstanceListDeviceColumnHeader", "Device"))
-									.FillWidth(0.3f)
-
-								+ SHeaderRow::Column("Status")
-									.DefaultLabel(LOCTEXT("InstanceListStatusColumnHeader", "Status"))
-									.FillWidth(0.2f)
-									.HAlignCell(HAlign_Right)
-									.HAlignHeader(HAlign_Right)
-							)
-					]
+				SAssignNew(SessionTreeView, STreeView<TSharedPtr<FSessionBrowserTreeItem>>)
+				.ItemHeight(20.0f)
+				.OnExpansionChanged(this, &SSessionBrowser::HandleSessionTreeViewExpansionChanged)
+				.OnGenerateRow(this, &SSessionBrowser::HandleSessionTreeViewGenerateRow)
+				.OnGetChildren(this, &SSessionBrowser::HandleSessionTreeViewGetChildren)
+				.OnSelectionChanged(this, &SSessionBrowser::HandleSessionTreeViewSelectionChanged)
+				.SelectionMode(ESelectionMode::Multi)
+				.TreeItemsSource(&SessionTreeItems)
+				.HeaderRow
+				(
+					SNew(SHeaderRow)
+
+					+ SHeaderRow::Column("Name")
+					.DefaultLabel(LOCTEXT("InstanceListNameColumnHeader", "Name"))
+					.FillWidth(0.3f)
+
+					+ SHeaderRow::Column("Type")
+					.DefaultLabel(LOCTEXT("InstanceListTypeColumnHeader", "Type"))
+					.FillWidth(0.2f)
+
+					+ SHeaderRow::Column("Device")
+					.DefaultLabel(LOCTEXT("InstanceListDeviceColumnHeader", "Device"))
+					.FillWidth(0.3f)
+
+					+ SHeaderRow::Column("Status")
+					.DefaultLabel(LOCTEXT("InstanceListStatusColumnHeader", "Status"))
+					.FillWidth(0.2f)
+					.HAlignCell(HAlign_Right)
+					.HAlignHeader(HAlign_Right)
+				)
 			]
-		/*
-		+ SVerticalBox::Slot()
-			.AutoHeight()
-			.Padding(0.0f, 4.0f, 0.0f, 0.0f)
-			[
-				SNew(SHorizontalBox)
-
-				+ SHorizontalBox::Slot()
-					.AutoWidth()
-					[
-						// terminate button
-						SNew(SButton)
-							.ButtonStyle(FEditorStyle::Get(), "ToggleButton")
-							.ContentPadding(FMargin(6.0f, 2.0f))
-							.IsEnabled(this, &SSessionBrowser::HandleTerminateSessionButtonIsEnabled)
-							.OnClicked(this, &SSessionBrowser::HandleTerminateSessionButtonClicked)
-							.ToolTipText(LOCTEXT("TerminateButtonTooltip", "Shuts down all game instances that are part of this session."))
-							[
-								SNew(SHorizontalBox)
-
-								+ SHorizontalBox::Slot()
-									.AutoWidth()
-									.VAlign(VAlign_Center)
-									[
-										SNew(SImage)
-											.Image(FEditorStyle::GetBrush("SessionBrowser.Terminate"))
-									]
-
-								+ SHorizontalBox::Slot()
-									.AutoWidth()
-									.VAlign(VAlign_Center)
-									.Padding(4.0f, 1.0f, 0.0f, 0.0f)
-									[
-										SNew(STextBlock)
-											.TextStyle(FEditorStyle::Get(), "SessionBrowser.Terminate.Font")
-											.Text(LOCTEXT("TerminateSessionButtonLabel", "Terminate Session"))
-									]
-							]
-					]
-			]*/
+		]
 	];
 
 	AppGroupItem = MakeShareable(new FSessionBrowserGroupTreeItem(LOCTEXT("AppGroupName", "This Application"), LOCTEXT("AppGroupToolTip", "The application instance that this session browser belongs to")));
@@ -237,6 +197,11 @@ void SSessionBrowser::FilterSessions()
 
 	// refresh tree view
 	SessionTreeView->RequestTreeRefresh();
+
+	if ( SessionTreeView->GetNumItemsSelected() == 0 && ThisAppInstance.IsValid() )
+	{
+		SessionTreeView->SetItemSelection(ThisAppInstance.Pin(), true, ESelectInfo::Direct);
+	}
 }
 
  void SSessionBrowser::AddInstanceItemToTree(TSharedPtr<FSessionBrowserTreeItem>& SessionItem, const TSharedPtr<FSessionBrowserTreeItem>& InstanceItem, const TSharedPtr<ISessionInstanceInfo>& InstanceInfo)
@@ -246,6 +211,8 @@ void SSessionBrowser::FilterSessions()
 	{
 		AppGroupItem->AddChild(InstanceItem.ToSharedRef());
 		InstanceItem->SetParent(AppGroupItem);
+
+		ThisAppInstance = InstanceItem;
 	}
 	else
 	{
@@ -298,7 +265,7 @@ void SSessionBrowser::HandleSessionManagerSelectedSessionChanged(const ISessionI
 		return;
 	}
 
-	IgnoreSessionTreeEvents = true;
+	updatingTreeExpansion = true;
 	{
 		if (SelectedSession.IsValid())
 		{
@@ -309,7 +276,7 @@ void SSessionBrowser::HandleSessionManagerSelectedSessionChanged(const ISessionI
 			SessionTreeView->SetSingleExpandedItem(nullptr);
 		}
 	}
-	IgnoreSessionTreeEvents = false;
+	updatingTreeExpansion = false;
 }
 
 void SSessionBrowser::HandleSessionManagerInstanceDiscovered(const TSharedRef<ISessionInfo>& OwnerSession, const TSharedRef<ISessionInstanceInfo>& DiscoveredInstance)
@@ -367,7 +334,7 @@ FText SSessionBrowser::HandleSessionTreeRowGetToolTipText(TSharedPtr<FSessionBro
 
 void SSessionBrowser::HandleSessionTreeViewExpansionChanged(TSharedPtr<FSessionBrowserTreeItem> TreeItem, bool bIsExpanded)
 {
-	if (IgnoreSessionTreeEvents || !TreeItem.IsValid())
+	if ( updatingTreeExpansion || !TreeItem.IsValid())
 	{
 		return;
 	}
@@ -381,11 +348,11 @@ void SSessionBrowser::HandleSessionTreeViewExpansionChanged(TSharedPtr<FSessionB
 	{
 		if (bIsExpanded)
 		{
-			IgnoreSessionTreeEvents = true;
+			updatingTreeExpansion = true;
 			{
 				ExpandItem(TreeItem);
 			}
-			IgnoreSessionTreeEvents = false;
+			updatingTreeExpansion = false;
 
 			// select session
 			if (TreeItem->GetType() == ESessionBrowserTreeNodeType::Session)
@@ -438,11 +405,6 @@ void SSessionBrowser::HandleSessionTreeViewGetChildren(TSharedPtr<FSessionBrowse
 
 void SSessionBrowser::HandleSessionTreeViewSelectionChanged(const TSharedPtr<FSessionBrowserTreeItem> Item, ESelectInfo::Type SelectInfo)
 {
-	if (IgnoreSessionTreeEvents || (SelectInfo == ESelectInfo::Direct))
-	{
-		return;
-	}
-
 	IgnoreSessionManagerEvents = true;
 	{
 		if (Item.IsValid())
diff --git a/Engine/Source/Developer/SessionFrontend/Private/Widgets/Browser/SSessionBrowser.h b/Engine/Source/Developer/SessionFrontend/Private/Widgets/Browser/SSessionBrowser.h
--- a/Engine/Source/Developer/SessionFrontend/Private/Widgets/Browser/SSessionBrowser.h
+++ b/Engine/Source/Developer/SessionFrontend/Private/Widgets/Browser/SSessionBrowser.h
@@ -94,7 +94,7 @@ private:
 	bool IgnoreSessionManagerEvents;
 
 	/** Whether to ignore events from the session tree view. */
-	bool IgnoreSessionTreeEvents;
+	bool updatingTreeExpansion;
 
 	/** Maps session and instance GUIDs to existing tree items. */
 	TMap<FGuid, TSharedPtr<FSessionBrowserTreeItem>> ItemMap;
@@ -121,4 +121,7 @@ private:
 
 	/** The session tree item that holds other user's standalone instances. */
 	TSharedPtr<FSessionBrowserGroupTreeItem> StandaloneGroupItem;
+
+	/** This app's instance session */
+	TWeakPtr<FSessionBrowserTreeItem> ThisAppInstance;
 };
diff --git a/Engine/Source/Developer/SessionFrontend/Private/Widgets/SSessionFrontend.h b/Engine/Source/Developer/SessionFrontend/Private/Widgets/SSessionFrontend.h
--- a/Engine/Source/Developer/SessionFrontend/Private/Widgets/SSessionFrontend.h
+++ b/Engine/Source/Developer/SessionFrontend/Private/Widgets/SSessionFrontend.h
@@ -25,6 +25,8 @@ public:
 	 */
 	void Construct( const FArguments& InArgs, const TSharedRef<SDockTab>& ConstructUnderMajorTab, const TSharedPtr<SWindow>& ConstructUnderWindow );
 
+	TSharedPtr<FTabManager> GetTabManager() const { return TabManager; }
+
 protected:
 
 	/**
diff --git a/Engine/Source/Developer/SessionFrontend/Public/ISessionFrontendModule.h b/Engine/Source/Developer/SessionFrontend/Public/ISessionFrontendModule.h
--- a/Engine/Source/Developer/SessionFrontend/Public/ISessionFrontendModule.h
+++ b/Engine/Source/Developer/SessionFrontend/Public/ISessionFrontendModule.h
@@ -33,6 +33,11 @@ public:
 	 */
 	virtual TSharedRef<SWidget> CreateSessionConsole( const TSharedRef<ISessionManager>& SessionManager ) = 0;
 
+	/**
+	 * Show the session frontend tab.
+	 */
+	virtual void InvokeSessionFrontend(FName SubTabToActivate = NAME_None) = 0;
+
 public:
 
 	/** Virtual destructor. */
diff --git a/Engine/Source/Developer/Settings/Private/SettingsContainer.cpp b/Engine/Source/Developer/Settings/Private/SettingsContainer.cpp
--- a/Engine/Source/Developer/Settings/Private/SettingsContainer.cpp
+++ b/Engine/Source/Developer/Settings/Private/SettingsContainer.cpp
@@ -94,9 +94,27 @@ int32 FSettingsContainer::GetCategories( TArray<ISettingsCategoryPtr>& OutCatego
 {
 	OutCategories.Empty(Categories.Num());
 
+	static const FName AdvancedCategoryName("Advanced");
+	TSharedPtr<FSettingsCategory> AdvancedCategory;
+
 	for (TMap<FName, TSharedPtr<FSettingsCategory> >::TConstIterator It(Categories); It; ++It)
 	{
-		OutCategories.Add(It.Value());
+		TSharedPtr<FSettingsCategory> Category = It.Value();
+		if(Category->GetName() == AdvancedCategoryName)
+		{
+			// Store off the advanced category, we'll add it to the bottom of all categories
+			AdvancedCategory = Category;
+		}
+		else
+		{
+			OutCategories.Add(It.Value());
+		}
+	}
+
+	// always show the advanced category last
+	if(AdvancedCategory.IsValid())
+	{
+		OutCategories.Add(AdvancedCategory);
 	}
 
 	return OutCategories.Num();
diff --git a/Engine/Source/Developer/Settings/Private/SettingsModule.cpp b/Engine/Source/Developer/Settings/Private/SettingsModule.cpp
--- a/Engine/Source/Developer/Settings/Private/SettingsModule.cpp
+++ b/Engine/Source/Developer/Settings/Private/SettingsModule.cpp
@@ -80,6 +80,7 @@ public:
 		EditorSettingsContainer->DescribeCategory("LevelEditor", LOCTEXT("EditorLevelEditorCategoryName", "Level Editor"), LOCTEXT("EditorLevelEditorCategoryDescription", "Level Editor settings description text here"));
 		EditorSettingsContainer->DescribeCategory("ContentEditors", LOCTEXT("EditorContentEditorsCategoryName", "Content Editors"), LOCTEXT("EditorContentEditorsCategoryDescription", "Content editors settings description text here"));
 		EditorSettingsContainer->DescribeCategory("Plugins", LOCTEXT("EditorPluginsCategoryName", "Plugins"), LOCTEXT("EditorPluginsCategoryDescription", "Plugins settings description text here"));
+		EditorSettingsContainer->DescribeCategory("Advanced", LOCTEXT("EditorAdvancedCategoryName", "Advanced"), LOCTEXT("EditorAdvancedCategoryDescription", "Advanced editor settings"));
 
 		// @todo gmp: move this into the modules that own these setting categories
 		TSharedRef<FSettingsContainer> ProjectSettingsContainer = FindOrAddContainer("Project");
diff --git a/Engine/Source/Developer/SettingsEditor/Private/Models/SettingsEditorModel.h b/Engine/Source/Developer/SettingsEditor/Private/Models/SettingsEditorModel.h
--- a/Engine/Source/Developer/SettingsEditor/Private/Models/SettingsEditorModel.h
+++ b/Engine/Source/Developer/SettingsEditor/Private/Models/SettingsEditorModel.h
@@ -58,6 +58,28 @@ public:
 		OnSelectionChangedDelegate.Broadcast();
 	}
 
+	virtual ISettingsSectionPtr GetSectionFromSectionObject(const UObject* SectionObject) const override
+	{
+		TArray<ISettingsCategoryPtr> Categories;
+		GetSettingsContainer()->GetCategories(Categories);
+
+		for(ISettingsCategoryPtr& Category : Categories)
+		{
+			TArray<ISettingsSectionPtr> Sections;
+			Category->GetSections(Sections);
+
+			for(ISettingsSectionPtr& Section : Sections)
+			{
+				if(Section->GetSettingsObject() == SectionObject)
+				{
+					return Section;
+				}
+			}
+		}
+
+		return nullptr;
+	}
+
 private:
 
 	/** Handles the removal of sections from the settings container. */
