Automation - Fixing some issues with the automation framework noticed while looking into a jira dealing with screenshots.  Metadata files were not being copied for new shots, and there was a bogus warning when walking the stack if the module had already been loaded.  Also added two more rendering tests to the Cornell Box map, Base Color and Scene Color.  Also fixing the screenshot for UIExample.

diff --git a/Engine/Source/Developer/FunctionalTesting/Private/AutomationBlueprintFunctionLibrary.cpp b/Engine/Source/Developer/FunctionalTesting/Private/AutomationBlueprintFunctionLibrary.cpp
--- a/Engine/Source/Developer/FunctionalTesting/Private/AutomationBlueprintFunctionLibrary.cpp
+++ b/Engine/Source/Developer/FunctionalTesting/Private/AutomationBlueprintFunctionLibrary.cpp
@@ -181,17 +181,17 @@ public:
 
 		if ( bWasNew )
 		{
-			UE_LOG(AutomationFunctionLibrary, Warning, TEXT("New Screenshot was discovered!  Please add a ground truth version of it."));
+			UE_LOG(AutomationFunctionLibrary, Warning, TEXT("New Screenshot '%s' was discovered!  Please add a ground truth version of it."), *Name);
 		}
 		else
 		{
 			if ( bWasSimilar )
 			{
-				UE_LOG(AutomationFunctionLibrary, Log, TEXT("Screenshot was similar!"));
+				UE_LOG(AutomationFunctionLibrary, Log, TEXT("Screenshot '%s' was similar!"), *Name);
 			}
 			else
 			{
-				UE_LOG(AutomationFunctionLibrary, Error, TEXT("Screenshot test failed, screenshots were different!"));
+				UE_LOG(AutomationFunctionLibrary, Error, TEXT("Screenshot '%s' test failed, screenshots were different!"), *Name);
 			}
 		}
 
@@ -260,6 +260,10 @@ bool UAutomationBlueprintFunctionLibrary::TakeAutomationScreenshotInternal(const
 	FAutomationScreenshotTaker* TempObject = new FAutomationScreenshotTaker(Name, Options);
 #endif
 
+	static IConsoleVariable* HighResScreenshotDelay = IConsoleManager::Get().FindConsoleVariable(TEXT("r.HighResScreenshotDelay"));
+	check(HighResScreenshotDelay);
+	HighResScreenshotDelay->Set(10);
+
     if ( FPlatformProperties::HasFixedResolution() )
     {
 	    FScreenshotRequest::RequestScreenshot(false);
diff --git a/Engine/Source/Developer/ScreenShotComparisonTools/Private/ScreenShotManager.cpp b/Engine/Source/Developer/ScreenShotComparisonTools/Private/ScreenShotManager.cpp
--- a/Engine/Source/Developer/ScreenShotComparisonTools/Private/ScreenShotManager.cpp
+++ b/Engine/Source/Developer/ScreenShotComparisonTools/Private/ScreenShotManager.cpp
@@ -267,7 +267,7 @@ FImageComparisonResult FScreenShotManager::CompareScreensot(FString ExistingImag
 	{
 		for ( FString ApprovedShot : ApprovedDeviceShots )
 		{
-			FString ApprovedShotFile = FPaths::GetCleanFilename(ApprovedDeviceShots[0]);
+			FString ApprovedShotFile = FPaths::GetCleanFilename(ApprovedShot);
 			FString ApprovedShotFileFull = FPaths::Combine(TestApprovedFolder, ApprovedShotFile);
 
 			FString ApprovedMetadataFile = FPaths::ChangeExtension(ApprovedShotFileFull, ".json");
diff --git a/Engine/Source/Runtime/Core/Private/Windows/WindowsPlatformStackWalk.cpp b/Engine/Source/Runtime/Core/Private/Windows/WindowsPlatformStackWalk.cpp
--- a/Engine/Source/Runtime/Core/Private/Windows/WindowsPlatformStackWalk.cpp
+++ b/Engine/Source/Runtime/Core/Private/Windows/WindowsPlatformStackWalk.cpp
@@ -585,7 +585,12 @@ static void LoadProcessModules(const FString &RemoteStorage)
 		if( !BaseAddress )
 		{
 			ErrorCode = GetLastError();
-			UE_LOG(LogWindows, Warning, TEXT("SymLoadModuleExW. Error: %d"), GetLastError());
+
+			// If the module is already loaded, the return value is zero and GetLastError returns ERROR_SUCCESS.
+			if ( ErrorCode != ERROR_SUCCESS )
+			{
+				UE_LOG(LogWindows, Warning, TEXT("SymLoadModuleExW. Error: %d"), ErrorCode);
+			}
 		}
 #else
 		SymSetSearchPath(ProcessHandle, TCHAR_TO_ANSI(*SearchPathList));
@@ -595,7 +600,12 @@ static void LoadProcessModules(const FString &RemoteStorage)
 		if (!BaseAddress)
 		{
 			ErrorCode = GetLastError();
-			UE_LOG(LogWindows, Warning, TEXT("SymLoadModuleEx. Error: %d"), GetLastError());
+
+			// If the module is already loaded, the return value is zero and GetLastError returns ERROR_SUCCESS.
+			if ( ErrorCode != ERROR_SUCCESS )
+			{
+				UE_LOG(LogWindows, Warning, TEXT("SymLoadModuleEx. Error: %d"), ErrorCode);
+			}
 		}
 #endif
 	} 
