Automation - Bringing a change over from Dev-Automation that fixes an issue with new screenshot tests that didn't apply the proper tolerence levels, all the new ones were defaulting to zero tolerence for pixel differences.

diff --git a/Engine/Source/Developer/FunctionalTesting/Private/ScreenshotFunctionalTest.cpp b/Engine/Source/Developer/FunctionalTesting/Private/ScreenshotFunctionalTest.cpp
--- a/Engine/Source/Developer/FunctionalTesting/Private/ScreenshotFunctionalTest.cpp
+++ b/Engine/Source/Developer/FunctionalTesting/Private/ScreenshotFunctionalTest.cpp
@@ -13,6 +13,7 @@
 
 AScreenshotFunctionalTest::AScreenshotFunctionalTest( const FObjectInitializer& ObjectInitializer )
 	: AFunctionalTest(ObjectInitializer)
+	, ScreenshotOptions(EComparisonTolerance::Low)
 {
 	ScreenshotCamera = CreateDefaultSubobject<UCameraComponent>(TEXT("Camera"));
 	ScreenshotCamera->SetupAttachment(RootComponent);
@@ -23,6 +24,9 @@ void AScreenshotFunctionalTest::PrepareTest()
 	APlayerController* PlayerController = UGameplayStatics::GetPlayerController(GetWorld(), 0);
 	PlayerController->SetViewTarget(this, FViewTargetTransitionParams());
 
+	// It's possible the defaults for certain tolerance levels have changed, so reset them on test start.
+	ScreenshotOptions.SetToleranceAmounts(ScreenshotOptions.Tolerance);
+
 	Super::PrepareTest();
 }
 
diff --git a/Engine/Source/Developer/FunctionalTesting/Public/AutomationScreenshotOptions.h b/Engine/Source/Developer/FunctionalTesting/Public/AutomationScreenshotOptions.h
--- a/Engine/Source/Developer/FunctionalTesting/Public/AutomationScreenshotOptions.h
+++ b/Engine/Source/Developer/FunctionalTesting/Public/AutomationScreenshotOptions.h
@@ -75,7 +75,7 @@ public:
 		, Delay(0.2f)
 		, bDisableNoisyRenderingFeatures(true)
 		, VisualizeBuffer(NAME_None)
-		, Tolerance(EComparisonTolerance::Low)
+		, Tolerance(EComparisonTolerance::Zero)
 		, ToleranceAmount()
 		, MaximumLocalError(0.10f)
 		, MaximumGlobalError(0.02f)
@@ -89,7 +89,7 @@ public:
 		, Delay(0.2f)
 		, bDisableNoisyRenderingFeatures(true)
 		, VisualizeBuffer(NAME_None)
-		, Tolerance(EComparisonTolerance::Low)
+		, Tolerance(InTolerance)
 		, ToleranceAmount()
 		, MaximumLocalError(0.10f)
 		, MaximumGlobalError(0.02f)
