Build - Fixing the editor compile on mac.

diff --git a/Engine/Source/Developer/FunctionalTesting/Private/AutomationBlueprintFunctionLibrary.cpp b/Engine/Source/Developer/FunctionalTesting/Private/AutomationBlueprintFunctionLibrary.cpp
--- a/Engine/Source/Developer/FunctionalTesting/Private/AutomationBlueprintFunctionLibrary.cpp
+++ b/Engine/Source/Developer/FunctionalTesting/Private/AutomationBlueprintFunctionLibrary.cpp
@@ -94,18 +94,6 @@ private:
 
 class FAutomationScreenshotTaker
 {
-	FString	Name;
-	FAutomationScreenshotOptions Options;
-
-	FConsoleVariableSwapper DefaultFeature_AntiAliasing;
-	FConsoleVariableSwapper DefaultFeature_AutoExposure;
-	FConsoleVariableSwapper DefaultFeature_MotionBlur;
-	FConsoleVariableSwapper PostProcessAAQuality;
-	FConsoleVariableSwapper MotionBlurQuality;
-	FConsoleVariableSwapper ScreenSpaceReflectionQuality;
-	FConsoleVariableSwapper EyeAdaptationQuality;
-	FConsoleVariableSwapper ContactShadows;
-
 public:
 	FAutomationScreenshotTaker(UWorld* InWorld, const FString& InName, FAutomationScreenshotOptions InOptions)
 		: World(InWorld)
@@ -146,16 +134,6 @@ public:
 
 		Options.SetToleranceAmounts(Options.Tolerance);
 
-		//if ( World.IsValid() )
-		//{
-		//	if ( ULocalPlayer* LocalPlayer = World->GetFirstLocalPlayerFromController() )
-		//	{
-		//		LocalPlayer->ResetViewState();
-		//	}
-		//}
-
-		//ADD_LATENT_AUTOMATION_COMMAND(FExecWorldStringLatentCommand(TEXT("r.ResetViewState")));
-
 		if ( UGameViewportClient* ViewportClient = GEngine->GameViewport )
 		{
 			static IConsoleVariable* ICVar = IConsoleManager::Get().FindConsoleVariable(FBufferVisualizationData::GetVisualizationTargetConsoleCommandName());
@@ -269,6 +247,18 @@ public:
 private:
 
 	TWeakObjectPtr<UWorld> World;
+	
+	FString	Name;
+	FAutomationScreenshotOptions Options;
+
+	FConsoleVariableSwapper DefaultFeature_AntiAliasing;
+	FConsoleVariableSwapper DefaultFeature_AutoExposure;
+	FConsoleVariableSwapper DefaultFeature_MotionBlur;
+	FConsoleVariableSwapper PostProcessAAQuality;
+	FConsoleVariableSwapper MotionBlurQuality;
+	FConsoleVariableSwapper ScreenSpaceReflectionQuality;
+	FConsoleVariableSwapper EyeAdaptationQuality;
+	FConsoleVariableSwapper ContactShadows;
 };
 
 #endif
