Fix CIS

diff --git a/Engine/Source/Developer/FunctionalTesting/Private/AutomationBlueprintFunctionLibrary.cpp b/Engine/Source/Developer/FunctionalTesting/Private/AutomationBlueprintFunctionLibrary.cpp
--- a/Engine/Source/Developer/FunctionalTesting/Private/AutomationBlueprintFunctionLibrary.cpp
+++ b/Engine/Source/Developer/FunctionalTesting/Private/AutomationBlueprintFunctionLibrary.cpp
@@ -493,10 +493,10 @@ void UAutomationBlueprintFunctionLibrary::DisableStatGroup(UObject* WorldContext
 #endif
 }
 
+#if STATS
 template <EComplexStatField::Type ValueType, bool bCallCount = false>
 float HelperGetStat(FName StatName)
 {
-#if STATS
 	if (FGameThreadStatsData* StatsData = FLatestGameThreadStatsData::Get().Latest)
 	{
 		if (const FComplexStatMessage* StatMessage = StatsData->GetStatData(StatName))
@@ -518,38 +518,49 @@ float HelperGetStat(FName StatName)
 	UE_LOG(AutomationFunctionLibrary, Warning, TEXT("%s"), *WarningOut.ToString());
 #endif
 
-#endif
 
 	return 0.f;
 }
+#endif
 
 float UAutomationBlueprintFunctionLibrary::GetStatIncAverage(FName StatName)
 {
+#if STATS
 	return HelperGetStat<EComplexStatField::IncAve>(StatName);
+#endif
+	return 0.f;
 }
 
 float UAutomationBlueprintFunctionLibrary::GetStatIncMax(FName StatName)
 {
-
+#if STATS
 	return HelperGetStat<EComplexStatField::IncMax>(StatName);
+#endif
+	return 0.f;
 }
 
 float UAutomationBlueprintFunctionLibrary::GetStatExcAverage(FName StatName)
 {
-
+#if STATS
 	return HelperGetStat<EComplexStatField::ExcAve>(StatName);
+#endif
+	return 0.f;
 }
 
 float UAutomationBlueprintFunctionLibrary::GetStatExcMax(FName StatName)
 {
-
+#if STATS
 	return HelperGetStat<EComplexStatField::ExcMax>(StatName);
+#endif
+	return 0.f;
 }
 
 float UAutomationBlueprintFunctionLibrary::GetStatCallCount(FName StatName)
 {
-
+#if STATS
 	return HelperGetStat<EComplexStatField::IncAve, /*bCallCount=*/true>(StatName);
+#endif
+	return 0.f;
 }
 
 bool UAutomationBlueprintFunctionLibrary::AreAutomatedTestsRunning()
