Copying //UE4/Dev-Automation to //UE4/Dev-Main (Source: //UE4/Dev-Automation @ 3448586)

==========================
MAJOR FEATURES + CHANGES
==========================

Change 3383163 on 2017/04/06 by james.fox

	Enabling Overlap automation tests.

Change 3387240 on 2017/04/10 by [REDACTED_PERSON]

	Enabling ProcMesh automation tests.  Made minor changes to in-level object names.

Change 3397994 on 2017/04/18 by [REDACTED_PERSON]

	Automation - Adding some areas for future extension to allow transmission of data, and performance information between controller and client.  The next step is to expose a way in blueprints, or via C++ to hook these events and send these messages.

Change 3401903 on 2017/04/20 by samuel.proctor

	Test map for blueprint bitwise operators

Change 3405592 on 2017/04/24 by [REDACTED_PERSON]

	Updated Roughness only example to be more visible in Decals test map

Change 3405625 on 2017/04/24 by [REDACTED_PERSON]

	Re-Saving material in decals map to remove log spam

Change 3406256 on 2017/04/24 by [REDACTED_PERSON]

	Show Collision tests, with Ground Truth Screenshots

Change 3410179 on 2017/04/26 by mason.seay

	Beginning of Math test map

Change 3412646 on 2017/04/27 by mason.seay

	More math tests

Change 3413556 on 2017/04/28 by [REDACTED_PERSON]

	Add support for expected errors in automation tests

	Automation tests can now define expected errors/warnings that will not prevent
	the test from succeeding. If expected messages are not encountered, the test
	will fail. The intent is to allow unit tests to verify negative or error cases
	while keeping the test pass report green. It is not intended as an error
	suppression system.

	#tests new automation tests

Change 3414217 on 2017/04/28 by mason.seay

	Wake Event functional test map

	Disabled Collision tests that still need review.

Change 3414605 on 2017/04/28 by mason.seay

	Added more tests to Math map

Change 3415806 on 2017/05/01 by [REDACTED_PERSON]

	Suspend log parsing during tests with expected errors

	Electric Commander was failing successful test passes that have expected errors
	because of its log parsing. BenM added markers to temporarily disable error
	parsing in CL 3414237, which this implements usage of.

	#tests preflight

Change 3416022 on 2017/05/01 by [REDACTED_PERSON]

	Physics/ChildTransform tests

Change 3420820 on 2017/05/03 by [REDACTED_PERSON]

	Automation - Allowing users to control the delay from the quick getter functions for Default Game/Rendering settings for screenshots.  Avoids needing to add more delay nodes to your graph.

	Improving the Box Drop example to use the new delay options.

	The FunctionalTest actor now ticks even when the game is paused.

Change 3422013 on 2017/05/03 by [REDACTED_PERSON]

	Automation - Adding an initial version of a ground truth storge mechanism that will only run in the editor.  Users can use special runs of their tests to store the ground truth, then rewire things to then in the future read that ground truth.  Added an example map showing the usage.

Change 3422802 on 2017/05/04 by [REDACTED_PERSON]

	Automation - Fixing a crash in the curve asset actions that assumes there will always be import data.

Change 3422803 on 2017/05/04 by [REDACTED_PERSON]

	Automation - Adding some comments to the ground truth function.

Change 3423739 on 2017/05/04 by [REDACTED_PERSON]

	Automation - Introducing a way for people to provide Ground Truth for their tests.

Change 3423768 on 2017/05/04 by [REDACTED_PERSON]

	Automation - GTD can no longer be modified on the build machine.  Unchecking modifiable on the example GTD.

Change 3425226 on 2017/05/05 by [REDACTED_PERSON]

	Automation - Fixing the FinishLoading code to only force shader compilation on platforms that don't require cooking.

Change 3428140 on 2017/05/08 by [REDACTED_PERSON]

	Added ground truth for scene query testing.

Change 3441024 on 2017/05/16 by [REDACTED_PERSON]

	Automation - Ground truth pass, improving system based on feedback.  Ditching the bCanModify, there's now a way to reset the data, and if it's been filled, you've got to reset it to save over it.

Change 3443894 on 2017/05/17 by [REDACTED_PERSON]

	Automation - Automation tweaking some logging.

Change 3446426 on 2017/05/18 by [REDACTED_PERSON]

	Automation - Adding additional logging and attempting to fix the warning on the build machine.

Change 3447617 on 2017/05/18 by [REDACTED_PERSON]

	Automation - Found an issue with the automation system keeping an old map around if the names were similar enough, and one contained the entirety of the other map's name.  This should fix the overlap warnings on the build machine.  Also cleaned up several places in the functional test manager that looked like they could cause confusion and problems with cleaning them up.

diff --git a/Engine/Plugins/Tests/FunctionalTestingEditor/Source/FunctionalTestingEditor.Build.cs b/Engine/Plugins/Tests/FunctionalTestingEditor/Source/FunctionalTestingEditor.Build.cs
--- a/Engine/Plugins/Tests/FunctionalTestingEditor/Source/FunctionalTestingEditor.Build.cs
+++ b/Engine/Plugins/Tests/FunctionalTestingEditor/Source/FunctionalTestingEditor.Build.cs
@@ -29,7 +29,8 @@ public class FunctionalTestingEditor : ModuleRules
                 "FunctionalTesting",
                 "PlacementMode",
                 "WorkspaceMenuStructure",
-                "ScreenShotComparisonTools"
+                "ScreenShotComparisonTools",
+				"AssetTools"
             }
          );
 	}
diff --git a/Engine/Plugins/Tests/FunctionalTestingEditor/Source/Private/AssetTypeActions_GroundTruthData.cpp b/Engine/Plugins/Tests/FunctionalTestingEditor/Source/Private/AssetTypeActions_GroundTruthData.cpp
new file mode 100644
--- /dev/null
+++ b/Engine/Plugins/Tests/FunctionalTestingEditor/Source/Private/AssetTypeActions_GroundTruthData.cpp
@@ -0,0 +1,12 @@
+// Copyright 1998-2017 Epic Games, Inc. All Rights Reserved.
+
+#include "AssetTypeActions_GroundTruthData.h"
+
+#define LOCTEXT_NAMESPACE "AssetTypeActions"
+
+void FAssetTypeActions_GroundTruthData::OpenAssetEditor(const TArray<UObject*>& InObjects, TSharedPtr<IToolkitHost> EditWithinLevelEditor)
+{
+	FSimpleAssetEditor::CreateEditor(EToolkitMode::Standalone, EditWithinLevelEditor, InObjects);
+}
+
+#undef LOCTEXT_NAMESPACE
diff --git a/Engine/Plugins/Tests/FunctionalTestingEditor/Source/Private/AssetTypeActions_GroundTruthData.h b/Engine/Plugins/Tests/FunctionalTestingEditor/Source/Private/AssetTypeActions_GroundTruthData.h
new file mode 100644
--- /dev/null
+++ b/Engine/Plugins/Tests/FunctionalTestingEditor/Source/Private/AssetTypeActions_GroundTruthData.h
@@ -0,0 +1,23 @@
+// Copyright 1998-2017 Epic Games, Inc. All Rights Reserved.
+
+#pragma once
+
+#include "CoreMinimal.h"
+#include "Toolkits/IToolkitHost.h"
+#include "AssetTypeActions_Base.h"
+#include "GroundTruthData.h"
+
+class FMenuBuilder;
+
+class FAssetTypeActions_GroundTruthData : public FAssetTypeActions_Base
+{
+public:
+	// IAssetTypeActions Implementation
+	virtual FColor GetTypeColor() const override { return FColor(255, 196, 128); }
+	virtual bool CanLocalize() const override { return false; }
+	virtual uint32 GetCategories() override { return EAssetTypeCategories::Misc; }
+	virtual FText GetName() const override { return NSLOCTEXT("AssetTypeActions", "AssetTypeActions_GroundTruthData", "Ground Truth Data"); }
+	virtual UClass* GetSupportedClass() const override { return UGroundTruthData::StaticClass(); }
+	virtual void OpenAssetEditor(const TArray<UObject*>& InObjects, TSharedPtr<class IToolkitHost> EditWithinLevelEditor = TSharedPtr<IToolkitHost>()) override;
+	// End IAssetTypeActions
+};
diff --git a/Engine/Plugins/Tests/FunctionalTestingEditor/Source/Private/EditorAutomationModule.cpp b/Engine/Plugins/Tests/FunctionalTestingEditor/Source/Private/EditorAutomationModule.cpp
--- a/Engine/Plugins/Tests/FunctionalTestingEditor/Source/Private/EditorAutomationModule.cpp
+++ b/Engine/Plugins/Tests/FunctionalTestingEditor/Source/Private/EditorAutomationModule.cpp
@@ -22,6 +22,8 @@
 #include "WorkspaceMenuStructure.h"
 #include "WorkspaceMenuStructureModule.h"
 #include "Widgets/Docking/SDockTab.h"
+#include "AssetTypeActions_GroundTruthData.h"
+#include "IAssetTools.h"
 
 #define LOCTEXT_NAMESPACE "EditorAutomation"
 
@@ -121,6 +123,9 @@ class FEditorAutomationModule : public IEditorAutomationModule
 		{
 			OnModulesChanged("PlacementMode", EModuleChangeReason::ModuleLoaded);
 		}
+
+		IAssetTools& AssetTools = FModuleManager::LoadModuleChecked<FAssetToolsModule>("AssetTools").Get();
+		AssetTools.RegisterAssetTypeActions(MakeShareable(new FAssetTypeActions_GroundTruthData()));
 	}
 
 	void ShutdownModule()
diff --git a/Engine/Plugins/Tests/FunctionalTestingEditor/Source/Private/GroundTruthDataFactory.cpp b/Engine/Plugins/Tests/FunctionalTestingEditor/Source/Private/GroundTruthDataFactory.cpp
new file mode 100644
--- /dev/null
+++ b/Engine/Plugins/Tests/FunctionalTestingEditor/Source/Private/GroundTruthDataFactory.cpp
@@ -0,0 +1,34 @@
+// Copyright 1998-2017 Epic Games, Inc. All Rights Reserved.
+
+#include "GroundTruthDataFactory.h"
+#include "GroundTruthData.h"
+#include "AssetTypeCategories.h"
+
+#define LOCTEXT_NAMESPACE "UGroundTruthDataFactory"
+
+UGroundTruthDataFactory::UGroundTruthDataFactory(const FObjectInitializer& ObjectInitializer)
+	: Super(ObjectInitializer)
+{
+	bCreateNew = true;
+	bEditAfterNew = false;
+	bEditorImport = false;
+	SupportedClass = UGroundTruthData::StaticClass();
+}
+
+UObject* UGroundTruthDataFactory::FactoryCreateNew(UClass* Class, UObject* InParent, FName InName, EObjectFlags InFlags, UObject* Context, FFeedbackContext* Warn)
+{
+	UGroundTruthData* GroundTruthData = NewObject<UGroundTruthData>(InParent, SupportedClass, InName, InFlags | RF_Transactional);
+	return GroundTruthData;
+}
+
+uint32 UGroundTruthDataFactory::GetMenuCategories() const
+{
+	return EAssetTypeCategories::Misc;
+}
+
+FText UGroundTruthDataFactory::GetDisplayName() const
+{
+	return LOCTEXT("MenuEntry", "Ground Truth Data");
+}
+
+#undef LOCTEXT_NAMESPACE
diff --git a/Engine/Plugins/Tests/FunctionalTestingEditor/Source/Private/GroundTruthDataFactory.h b/Engine/Plugins/Tests/FunctionalTestingEditor/Source/Private/GroundTruthDataFactory.h
new file mode 100644
--- /dev/null
+++ b/Engine/Plugins/Tests/FunctionalTestingEditor/Source/Private/GroundTruthDataFactory.h
@@ -0,0 +1,20 @@
+// Copyright 1998-2017 Epic Games, Inc. All Rights Reserved.
+
+#pragma once
+
+#include "CoreMinimal.h"
+#include "UObject/ObjectMacros.h"
+#include "Factories/Factory.h"
+#include "GroundTruthDataFactory.generated.h"
+
+UCLASS(hidecategories=Object, collapsecategories)
+class UGroundTruthDataFactory : public UFactory
+{
+	GENERATED_UCLASS_BODY()
+
+	// UFactory Interface
+	virtual UObject* FactoryCreateNew(UClass* Class, UObject* InParent, FName Name, EObjectFlags Flags, UObject* Context, FFeedbackContext* Warn) override;
+	virtual uint32 GetMenuCategories() const override;
+	virtual FText GetDisplayName() const override;
+	// End of UFactory Interface
+};
diff --git a/Engine/Plugins/Tests/RuntimeTests/Source/RuntimeTests/Private/AutomationExpectedError.spec.cpp b/Engine/Plugins/Tests/RuntimeTests/Source/RuntimeTests/Private/AutomationExpectedError.spec.cpp
new file mode 100644
--- /dev/null
+++ b/Engine/Plugins/Tests/RuntimeTests/Source/RuntimeTests/Private/AutomationExpectedError.spec.cpp
@@ -0,0 +1,172 @@
+// Copyright 1998-2017 Epic Games, Inc. All Rights Reserved.
+
+#include "AutomationTest.h"
+
+BEGIN_DEFINE_SPEC(FAutomationExpectedErrorTest, "System.Automation.ExpectedError", EAutomationTestFlags::EngineFilter | EAutomationTestFlags::ApplicationContextMask)
+END_DEFINE_SPEC(FAutomationExpectedErrorTest)
+void FAutomationExpectedErrorTest::Define()
+{
+	Describe("A defined expected error in a test", [this]()
+	{
+
+		It("will not add an error with a number of occurrences less than zero", [this]()
+		{
+			// Suppress error logged when adding entry with invalid occurrence count
+			AddExpectedError(TEXT("number of expected occurrences must be >= 0"), EAutomationExpectedErrorFlags::Contains, 1);
+
+			AddExpectedError(TEXT("The two values are not equal"), EAutomationExpectedErrorFlags::Contains, -1);
+
+			TArray<FAutomationExpectedError> Errors;
+			GetExpectedErrors(Errors);
+
+			// Only the first expected error should exist in the list.
+			TestEqual("Expected Errors Count", Errors.Num(), 1);
+		});
+
+		It("will add an error with a number of occurrences equal to zero", [this]()
+		{
+			AddExpectedError(TEXT("Expected Error"), EAutomationExpectedErrorFlags::Contains, 0);
+
+			TArray<FAutomationExpectedError> Errors;
+			GetExpectedErrors(Errors);
+			TestEqual("Expected Errors Count", Errors.Num(), 1);
+
+			// Add the expected error to ensure all test conditions pass
+			AddError(TEXT("Expected Error"));
+		});
+
+		It("will not duplicate an existing expected error using the same matcher", [this]()
+		{
+			// Suppress warning logged when adding duplicate value
+			AddExpectedError(TEXT("cannot add duplicate entries"), EAutomationExpectedErrorFlags::Contains, 1);
+			
+			AddExpectedError(TEXT("Expected Error"), EAutomationExpectedErrorFlags::Contains, 1);
+			AddExpectedError(TEXT("Expected Error"), EAutomationExpectedErrorFlags::Contains, 1);
+
+			TArray<FAutomationExpectedError> Errors;
+			GetExpectedErrors(Errors);
+			TestEqual("Expected Errors Count", Errors.Num(), 2);
+
+			// Add the expected error to ensure all test conditions pass
+			AddError(TEXT("Expected Error"));
+		});
+
+		It("will not duplicate an expected error using a different matcher", [this]()
+		{
+			// Suppress warning logged when adding duplicate value
+			AddExpectedError(TEXT("cannot add duplicate entries"), EAutomationExpectedErrorFlags::Contains, 2);
+
+			AddExpectedError(TEXT("Expected Exact Error"), EAutomationExpectedErrorFlags::Exact, 1);
+			AddExpectedError(TEXT("Expected Exact Error"), EAutomationExpectedErrorFlags::Contains, 1);
+
+			AddExpectedError(TEXT("Expected Contains Error"), EAutomationExpectedErrorFlags::Contains, 1);
+			AddExpectedError(TEXT("Expected Contains Error"), EAutomationExpectedErrorFlags::Exact, 1);
+
+			TArray<FAutomationExpectedError> Errors;
+			GetExpectedErrors(Errors);
+			TestEqual("Expected Errors Count", Errors.Num(), 3);
+
+			// Add the expected errors to ensure all test conditions pass
+			AddError(TEXT("Expected Exact Error"));
+			AddError(TEXT("Expected Contains Error"));
+		});
+
+		// Disabled until fix for UE-44340 (crash creating invalid regex) is merged
+		xIt("will not add an error with an invalid regex pattern", [this]()
+		{
+			AddExpectedError(TEXT("invalid regex }])([{"), EAutomationExpectedErrorFlags::Contains, 0);
+
+			TArray<FAutomationExpectedError> Errors;
+			GetExpectedErrors(Errors);
+
+			TestEqual("Expected Errors Count", Errors.Num(), 0);
+		});
+
+		It("will match both Error and Warning messages", [this]()
+		{
+			AddExpectedError(TEXT("Expected Message"), EAutomationExpectedErrorFlags::Contains, 0);
+			AddError(TEXT("Expected Message"));
+			AddWarning(TEXT("Expected Message"));
+		});
+
+		It("will not fail or warn if encountered", [this]()
+		{
+			AddExpectedError(TEXT("Expected Error"), EAutomationExpectedErrorFlags::Contains, 1);
+			AddExpectedError(TEXT("Expected Warning"), EAutomationExpectedErrorFlags::Contains, 1);
+			AddError(TEXT("Expected Error"));
+			AddWarning(TEXT("Expected Warning"));
+		});
+
+		It("will not match multiple occurrences in the same message when using Contains matcher", [this]()
+		{
+			AddExpectedError(TEXT("Expected"), EAutomationExpectedErrorFlags::Contains, 1);
+			AddError(TEXT("ExpectedExpectedExpected ExpectedExpectedExpected"));
+		});
+
+		It("will match different messages that fit the regex pattern", [this]()
+		{
+			AddExpectedError(TEXT("Response \\(-?\\d+\\)"), EAutomationExpectedErrorFlags::Contains, 4);
+			AddError(TEXT("Response (0)"));
+			AddError(TEXT("Response (1)"));
+			AddError(FString::Printf(TEXT("Response (%d)"), MIN_int64));
+			AddError(FString::Printf(TEXT("Response (%d)"), MAX_uint64));
+		});
+
+	});
+
+}
+
+// Tests for cases where expected errors will fail a test.
+// IMPORTANT: The pass condition for these tests is that they FAIL. To prevent
+// the expected failures from interfering with regular test runs, these tests
+// must be run manually.
+BEGIN_DEFINE_SPEC(FAutomationExpectedErrorFailureTest, "System.Automation.ExpectedError", EAutomationTestFlags::EngineFilter | EAutomationTestFlags::ApplicationContextMask | EAutomationTestFlags::RequiresUser)
+END_DEFINE_SPEC(FAutomationExpectedErrorFailureTest)
+
+void FAutomationExpectedErrorFailureTest::Define()
+{
+	Describe("An expected error failure", [this]()
+	{
+		It("will occur if expected a specific number of times and NOT encountered.", [this]()
+		{
+			AddExpectedError(TEXT("Expected Error"), EAutomationExpectedErrorFlags::Exact, 1);
+		});
+
+		It("will occur if expected a specific number of times and is encountered too few times.", [this]()
+		{
+			AddExpectedError(TEXT("Expected Error"), EAutomationExpectedErrorFlags::Exact, 2);
+			AddError(TEXT("Expected Error"));
+		});
+
+		It("will occur if expected a specific number of times and is encountered too many times.", [this]()
+		{
+			AddExpectedError(TEXT("Expected Error"), EAutomationExpectedErrorFlags::Exact, 1);
+			AddError(TEXT("Expected Error"));
+			AddError(TEXT("Expected Error"));
+		});
+
+		It("will occur if expected any number of times and is not encountered.", [this]()
+		{
+			AddExpectedError(TEXT("Expected Error"), EAutomationExpectedErrorFlags::Exact, 0);
+		});
+
+		It("will occur if multiple expected errors are NOT encountered.", [this]()
+		{
+			AddExpectedError(TEXT("Expected Error 1"), EAutomationExpectedErrorFlags::Exact, 1);
+			AddExpectedError(TEXT("Expected Error 2"), EAutomationExpectedErrorFlags::Contains, 1);
+		});
+
+		It("will occur if not all expected errors are encountered.", [this]()
+		{
+			AddExpectedError(TEXT("Expected error 1"), EAutomationExpectedErrorFlags::Exact, 1);
+			AddExpectedError(TEXT("Expected error 2"), EAutomationExpectedErrorFlags::Contains, 1);
+			AddError(TEXT("Expected error 1"));
+		});
+
+		It("will occur if only partial matches are encountered when using Exact matcher", [this]()
+		{
+			AddExpectedError(TEXT("Expected"), EAutomationExpectedErrorFlags::Exact, 1);
+			AddError(TEXT("Expected error"));
+		});
+	});
+}
\ No newline at end of file
diff --git a/Engine/Source/Developer/AssetTools/Private/AssetTypeActions/AssetTypeActions_Curve.cpp b/Engine/Source/Developer/AssetTools/Private/AssetTypeActions/AssetTypeActions_Curve.cpp
--- a/Engine/Source/Developer/AssetTools/Private/AssetTypeActions/AssetTypeActions_Curve.cpp
+++ b/Engine/Source/Developer/AssetTools/Private/AssetTypeActions/AssetTypeActions_Curve.cpp
@@ -12,7 +12,7 @@ void FAssetTypeActions_Curve::OpenAssetEditor( const TArray<UObject*>& InObjects
 	for (auto ObjIt = InObjects.CreateConstIterator(); ObjIt; ++ObjIt)
 	{
 		auto Curve = Cast<UCurveBase>(*ObjIt);
-		if (Curve != NULL)
+		if (Curve != nullptr)
 		{
 			FCurveAssetEditorModule& CurveAssetEditorModule = FModuleManager::LoadModuleChecked<FCurveAssetEditorModule>( "CurveAssetEditor" );
 			TSharedRef< ICurveAssetEditor > NewCurveAssetEditor = CurveAssetEditorModule.CreateCurveAssetEditor( Mode, EditWithinLevelEditor, Curve );
@@ -25,6 +25,9 @@ void FAssetTypeActions_Curve::GetResolvedSourceFilePaths(const TArray<UObject*>&
 	for (auto& Asset : TypeAssets)
 	{
 		const auto Curve = CastChecked<UCurveBase>(Asset);
-		Curve->AssetImportData->ExtractFilenames(OutSourceFilePaths);
+		if (Curve->AssetImportData)
+		{
+			Curve->AssetImportData->ExtractFilenames(OutSourceFilePaths);
+		}
 	}
 }
diff --git a/Engine/Source/Developer/AutomationController/Private/AutomationControllerManager.h b/Engine/Source/Developer/AutomationController/Private/AutomationControllerManager.h
--- a/Engine/Source/Developer/AutomationController/Private/AutomationControllerManager.h
+++ b/Engine/Source/Developer/AutomationController/Private/AutomationControllerManager.h
@@ -403,6 +403,12 @@ private:
 	/** Handles FAutomationWorkerScreenImage messages. */
 	void HandleReceivedScreenShot( const FAutomationWorkerScreenImage& Message, const IMessageContextRef& Context );
 
+	/** Handles FAutomationWorkerTestDataRequest messages. */
+	void HandleTestDataRequest(const FAutomationWorkerTestDataRequest& Message, const IMessageContextRef& Context);
+
+	/** Handles FAutomationWorkerPerformanceDataRequest messages. */
+	void HandlePerformanceDataRequest(const FAutomationWorkerPerformanceDataRequest& Message, const IMessageContextRef& Context);
+
 	/** Handles FAutomationWorkerRequestNextNetworkCommand messages. */
 	void HandleRequestNextNetworkCommandMessage( const FAutomationWorkerRequestNextNetworkCommand& Message, const IMessageContextRef& Context );
 
diff --git a/Engine/Source/Developer/AutomationController/Private/AutomationControllerManger.cpp b/Engine/Source/Developer/AutomationController/Private/AutomationControllerManger.cpp
--- a/Engine/Source/Developer/AutomationController/Private/AutomationControllerManger.cpp
+++ b/Engine/Source/Developer/AutomationController/Private/AutomationControllerManger.cpp
@@ -500,6 +500,7 @@ void FAutomationControllerManager::Startup()
 		.Handling<FAutomationWorkerRequestTestsReplyComplete>(this, &FAutomationControllerManager::HandleRequestTestsReplyCompleteMessage)
 		.Handling<FAutomationWorkerRunTestsReply>(this, &FAutomationControllerManager::HandleRunTestsReplyMessage)
 		.Handling<FAutomationWorkerScreenImage>(this, &FAutomationControllerManager::HandleReceivedScreenShot)
+		.Handling<FAutomationWorkerTestDataRequest>(this, &FAutomationControllerManager::HandleTestDataRequest)
 		.Handling<FAutomationWorkerWorkerOffline>(this, &FAutomationControllerManager::HandleWorkerOfflineMessage);
 
 	if ( MessageEndpoint.IsValid() )
@@ -898,6 +899,62 @@ void FAutomationControllerManager::HandleReceivedScreenShot(const FAutomationWor
 	ComparisonQueue.Enqueue(Comparison);
 }
 
+void FAutomationControllerManager::HandleTestDataRequest(const FAutomationWorkerTestDataRequest& Message, const IMessageContextRef& Context)
+{
+	const FString TestDataRoot = FPaths::ConvertRelativePathToFull(FPaths::GameDir() / TEXT("Test"));
+	const FString DataFile = Message.DataType / Message.DataPlatform / Message.DataTestName / Message.DataName + TEXT(".json");
+	const FString DataFullPath = TestDataRoot / DataFile;
+
+	// Generate the folder for the data if it doesn't exist.
+	const bool bTree = true;
+	IFileManager::Get().MakeDirectory(*FPaths::GetPath(DataFile), bTree);
+
+	bool bIsNew = true;
+	FString ResponseJsonData = Message.JsonData;
+
+	if ( FPaths::FileExists(DataFullPath) )
+	{
+		if ( FFileHelper::LoadFileToString(ResponseJsonData, *DataFullPath) )
+		{
+			bIsNew = false;
+		}
+		else
+		{
+			// TODO Error
+		}
+	}
+
+	if ( bIsNew )
+	{
+		FString IncomingTestData = FPaths::GameSavedDir() / TEXT("Automation/IncomingData/") / DataFile;
+		if ( FFileHelper::SaveStringToFile(Message.JsonData, *IncomingTestData) )
+		{
+			//TODO Anything extra to do here?
+		}
+		else
+		{
+			//TODO What do we do if this fails?
+		}
+	}
+
+	FAutomationWorkerTestDataResponse* ResponseMessage = new FAutomationWorkerTestDataResponse();
+	ResponseMessage->bIsNew = bIsNew;
+	ResponseMessage->JsonData = ResponseJsonData;
+
+	MessageEndpoint->Send(ResponseMessage, Context->GetSender());
+}
+
+void FAutomationControllerManager::HandlePerformanceDataRequest(const FAutomationWorkerPerformanceDataRequest& Message, const IMessageContextRef& Context)
+{
+	//TODO Read/Performance data.
+
+	FAutomationWorkerPerformanceDataResponse* ResponseMessage = new FAutomationWorkerPerformanceDataResponse();
+	ResponseMessage->bSuccess = true;
+	ResponseMessage->ErrorMessage = TEXT("");
+
+	MessageEndpoint->Send(ResponseMessage, Context->GetSender());
+}
+
 void FAutomationControllerManager::HandleRequestNextNetworkCommandMessage(const FAutomationWorkerRequestNextNetworkCommand& Message, const IMessageContextRef& Context)
 {
 	// Harvest iteration of running the tests this result came from (stops stale results from being committed to subsequent runs)
diff --git a/Engine/Source/Developer/FunctionalTesting/Classes/FunctionalTest.h b/Engine/Source/Developer/FunctionalTesting/Classes/FunctionalTest.h
--- a/Engine/Source/Developer/FunctionalTesting/Classes/FunctionalTest.h
+++ b/Engine/Source/Developer/FunctionalTesting/Classes/FunctionalTest.h
@@ -13,6 +13,7 @@
 
 class Error;
 class UBillboardComponent;
+class UTraceQueryTestResults;
 
 //Experimental effort at automated cpu captures from the functional testing.
 class FFunctionalTestExternalProfiler : public FScopedExternalProfilerBase
@@ -363,112 +364,133 @@ public:
 	 * @param Message	The message to display if the assert fails ("Assertion Failed: 'Message' for context ''")
 	 */
 	UFUNCTION(BlueprintCallable, Category = "Asserts", meta = ( HidePin = "ContextObject", DefaultToSelf = "ContextObject"))
-	void AssertTrue(bool Condition, FString Message, const UObject* ContextObject = nullptr);
+	bool AssertTrue(bool Condition, FString Message, const UObject* ContextObject = nullptr);
 
 	/**
 	 * Assert that a boolean value is false.
 	 * @param Message	The message to display if the assert fails ("Assertion Failed: 'Message' for context ''")
 	 */
 	UFUNCTION(BlueprintCallable, Category = "Asserts", meta = ( HidePin = "ContextObject", DefaultToSelf = "ContextObject"))
-	void AssertFalse(bool Condition, FString Message, const UObject* ContextObject = nullptr);
+	bool AssertFalse(bool Condition, FString Message, const UObject* ContextObject = nullptr);
 
 	/**
 	 * Assert that a UObject is valid
 	 * @param Message	The message to display if the object is invalid ("Invalid object: 'Message' for context ''")
 	 */
 	UFUNCTION(BlueprintCallable, Category = "Asserts", meta = ( HidePin = "ContextObject", DefaultToSelf = "ContextObject"))
-	void AssertIsValid(UObject* Object, FString Message, const UObject* ContextObject = nullptr);
+	bool AssertIsValid(UObject* Object, FString Message, const UObject* ContextObject = nullptr);
 
 	/**
 	 * Assert on a relationship between two integers.
 	 * @param What	A name to use in the message if the assert fails (What: expected {Actual} to be <ShouldBe> {Expected} for context '')
 	 */
 	UFUNCTION(BlueprintCallable, Category = "Asserts", DisplayName = "Assert Value (int)", meta = ( HidePin = "ContextObject", DefaultToSelf = "ContextObject"))
-	void AssertValue_Int(int32 Actual, EComparisonMethod ShouldBe, int32 Expected, const FString& What, const UObject* ContextObject = nullptr);
+	bool AssertValue_Int(int32 Actual, EComparisonMethod ShouldBe, int32 Expected, const FString& What, const UObject* ContextObject = nullptr);
 
 	/**
 	 * Assert on a relationship between two floats.
 	 * @param What	A name to use in the message if the assert fails (What: expected {Actual} to be <ShouldBe> {Expected} for context '')
 	 */
 	UFUNCTION(BlueprintCallable, Category = "Asserts", DisplayName = "Assert Value (float)", meta = ( HidePin = "ContextObject", DefaultToSelf = "ContextObject"))
-	void AssertValue_Float(float Actual, EComparisonMethod ShouldBe, float Expected, const FString& What, const UObject* ContextObject = nullptr);
+	bool AssertValue_Float(float Actual, EComparisonMethod ShouldBe, float Expected, const FString& What, const UObject* ContextObject = nullptr);
 
 	/**
 	 * Assert on a relationship between two DateTimes.
 	 * @param What	A name to use in the message if the assert fails (What: expected {Actual} to be <ShouldBe> {Expected} for context '')
 	 */
 	UFUNCTION(BlueprintCallable, Category = "Asserts", DisplayName = "Assert Value (DateTime)", meta = ( HidePin = "ContextObject", DefaultToSelf = "ContextObject"))
-	void AssertValue_DateTime(FDateTime Actual, EComparisonMethod ShouldBe, FDateTime Expected, const FString& What, const UObject* ContextObject = nullptr);
+	bool AssertValue_DateTime(FDateTime Actual, EComparisonMethod ShouldBe, FDateTime Expected, const FString& What, const UObject* ContextObject = nullptr);
 
 	/**
 	 * Assert that two transforms are (components memberwise - translation, rotation, scale) equal within a small tolerance.
 	 * @param What	A name to use in the message if the assert fails ("Expected 'What' to be {Expected} but it was {Actual} for context ''")
 	 */
 	UFUNCTION(BlueprintCallable, Category = "Asserts", DisplayName = "Assert Equal (Transform)", meta = ( HidePin = "ContextObject", DefaultToSelf = "ContextObject"))
-	void AssertEqual_Transform(const FTransform Actual, const FTransform Expected, const FString& What, const UObject* ContextObject = nullptr);
+	bool AssertEqual_Transform(const FTransform Actual, const FTransform Expected, const FString& What, const UObject* ContextObject = nullptr);
 
 	/**
 	 * Assert that two floats are equal within tolerance between two floats.
 	 * @param What	A name to use in the message if the assert fails (What: expected {Actual} to be Equal To {Expected} within Tolerance for context '')
 	 */
 	UFUNCTION(BlueprintCallable, Category = "Asserts", DisplayName = "Assert Equal (Float)", meta = ( HidePin = "ContextObject", DefaultToSelf = "ContextObject"))
-	void AssertEqual_Float(const float Actual, const float Expected, const FString& What, float Tolerance = 1.e-4, const UObject* ContextObject = nullptr);
+	bool AssertEqual_Float(const float Actual, const float Expected, const FString& What, float Tolerance = 1.e-4, const UObject* ContextObject = nullptr);
 
 	/**
-	* Assert that two integers are equal.
-	* @param What	A name to use in the message if the assert fails (What: expected {Actual} to be Equal To {Expected} '')
+	* Assert that two bools are equal
+	* @param What	A name to use in the message if the assert fails (What: expected {Actual} to be Equal To {Expected} for context '')
+	*/
+	UFUNCTION(BlueprintCallable, Category = "Asserts", DisplayName = "Assert Equal (Bool)", meta = (HidePin = "ContextObject", DefaultToSelf = "ContextObject"))
+	bool AssertEqual_Bool(const bool Actual, const bool Expected, const FString& What, const UObject* ContextObject = nullptr);
+
+	/**
+	* Assert that two ints are equal
+	* @param What	A name to use in the message if the assert fails (What: expected {Actual} to be Equal To {Expected} for context '')
 	*/
 	UFUNCTION(BlueprintCallable, Category = "Asserts", DisplayName = "Assert Equal (Int)", meta = (HidePin = "ContextObject", DefaultToSelf = "ContextObject"))
-	bool AssertEqual_Int(const int32 Actual, const int32 Expected, const FString& What, const UObject* ContextObject = nullptr);
+	bool AssertEqual_Int(const int Actual, const int Expected, const FString& What, const UObject* ContextObject = nullptr);
+
+	/**
+	* Assert that two FNames are equal
+	* @param What	A name to use in the message if the assert fails (What: expected {Actual} to be Equal To {Expected} for context '')
+	*/
+	UFUNCTION(BlueprintCallable, Category = "Asserts", DisplayName = "Assert Equal (FName)", meta = (HidePin = "ContextObject", DefaultToSelf = "ContextObject"))
+	bool AssertEqual_Name(const FName Actual, const FName Expected, const FString& What, const UObject* ContextObject = nullptr);
 
 	/**
 	 * Assert that two transforms are (components memberwise - translation, rotation, scale) not equal within a small tolerance.
 	 * @param What	A name to use in the message if the assert fails ("Expected 'What' not to be {Expected} but it was {Actual} for context ''")
 	 */
 	UFUNCTION(BlueprintCallable, Category = "Asserts", DisplayName = "Assert Not Equal (Transform)", meta = ( HidePin = "ContextObject", DefaultToSelf = "ContextObject"))
-	void AssertNotEqual_Transform(const FTransform Actual, const FTransform NotExpected, const FString& What, const UObject* ContextObject = nullptr);
+	bool AssertNotEqual_Transform(const FTransform Actual, const FTransform NotExpected, const FString& What, const UObject* ContextObject = nullptr);
 
 	/**
 	 * Assert that the component angles of two rotators are all equal within a small tolerance.
 	 * @param What	A name to use in the message if the assert fails ("Expected 'What' to be {Expected} but it was {Actual} for context ''")
 	 */
 	UFUNCTION(BlueprintCallable, Category = "Asserts", DisplayName = "Assert Equal (Rotator)", meta = ( HidePin = "ContextObject", DefaultToSelf = "ContextObject"))
-	void AssertEqual_Rotator(const FRotator Actual, const FRotator Expected, const FString& What, const UObject* ContextObject = nullptr);
+	bool AssertEqual_Rotator(const FRotator Actual, const FRotator Expected, const FString& What, const UObject* ContextObject = nullptr);
 
 	/**
 	 * Assert that the component angles of two rotators are all not equal within a small tolerance.
 	 * @param What	A name to use in the message if the assert fails ("Expected 'What' not to be {Expected} but it was {Actual} for context ''")
 	 */
 	UFUNCTION(BlueprintCallable, Category = "Asserts", DisplayName = "Assert Not Equal (Rotator)", meta = ( HidePin = "ContextObject", DefaultToSelf = "ContextObject"))
-	void AssertNotEqual_Rotator(const FRotator Actual, const FRotator NotExpected, const FString& What, const UObject* ContextObject = nullptr);
+	bool AssertNotEqual_Rotator(const FRotator Actual, const FRotator NotExpected, const FString& What, const UObject* ContextObject = nullptr);
 
 	/**
 	 * Assert that two vectors are (memberwise) equal within a small tolerance.
 	 * @param What	A name to use in the message if the assert fails ("Expected 'What' to be {Expected} but it was {Actual} for context ''")
 	 */
 	UFUNCTION(BlueprintCallable, Category = "Asserts", DisplayName = "Assert Equal (Vector)", meta = ( HidePin = "ContextObject", DefaultToSelf = "ContextObject"))
-	void AssertEqual_Vector(const FVector Actual, const FVector Expected, const FString& What, const float Tolerance = 1.e-4f, const UObject* ContextObject = nullptr);
+	bool AssertEqual_Vector(const FVector Actual, const FVector Expected, const FString& What, const float Tolerance = 1.e-4f, const UObject* ContextObject = nullptr);
 
 	/**
 	 * Assert that two vectors are (memberwise) not equal within a small tolerance.
 	 * @param What	A name to use in the message if the assert fails ("Expected 'What' not to be {Expected} but it was {Actual} for context ''")
 	 */
 	UFUNCTION(BlueprintCallable, Category = "Asserts", DisplayName = "Assert Not Equal (Vector)", meta = ( HidePin = "ContextObject", DefaultToSelf = "ContextObject"))
-	void AssertNotEqual_Vector(const FVector Actual, const FVector NotExpected, const FString& What, const UObject* ContextObject = nullptr);
+	bool AssertNotEqual_Vector(const FVector Actual, const FVector NotExpected, const FString& What, const UObject* ContextObject = nullptr);
 
 	/**
 	 * Assert that two Strings are equal.
 	 * @param What	A name to use in the message if the assert fails ("Expected 'What' to be {Expected} but it was {Actual} for context ''")
 	 */
 	UFUNCTION(BlueprintCallable, Category = "Asserts", DisplayName = "Assert Equal (String)", meta = ( HidePin = "ContextObject", DefaultToSelf = "ContextObject"))
-	void AssertEqual_String(const FString Actual, const FString Expected, const FString& What, const UObject* ContextObject = nullptr);
+	bool AssertEqual_String(const FString Actual, const FString Expected, const FString& What, const UObject* ContextObject = nullptr);
 
 	/**
 	 * Assert that two Strings are not equal.
 	 * @param What	A name to use in the message if the assert fails ("Expected 'What' not to be {Expected} but it was {Actual} for context ''")
 	 */
 	UFUNCTION(BlueprintCallable, Category = "Asserts", DisplayName = "Assert Not Equal (String)", meta = ( HidePin = "ContextObject", DefaultToSelf = "ContextObject"))
-	void AssertNotEqual_String(const FString Actual, const FString NotExpected, const FString& What, const UObject* ContextObject = nullptr);
+	bool AssertNotEqual_String(const FString Actual, const FString NotExpected, const FString& What, const UObject* ContextObject = nullptr);
+
+	/**
+	* Assert that two TraceQueryResults are equal.
+	* @param What	A name to use in the message if the assert fails ("Expected 'What' not to be {Expected} but it was {Actual} for context ''")
+	*/
+	UFUNCTION(BlueprintCallable, Category = "Asserts", DisplayName = "Assert Equal (TraceQuery)", meta = (HidePin = "ContextObject", DefaultToSelf = "ContextObject"))
+	bool AssertEqual_TraceQueryResults(const UTraceQueryTestResults* Actual, const UTraceQueryTestResults* Expected, const FString& What, const UObject* ContextObject = nullptr);
 
 	UFUNCTION(BlueprintCallable, Category = "Reporting")
 	void AddWarning(const FString Message);
@@ -476,7 +498,8 @@ public:
 	UFUNCTION(BlueprintCallable, Category = "Reporting")
 	void AddError(const FString Message);
 
-protected:
+//protected:
+	/** TODO: break this out into a library */
 	void LogStep(ELogVerbosity::Type Verbosity, const FString& Message);
 
 public:
diff --git a/Engine/Source/Developer/FunctionalTesting/Classes/FunctionalTestUtilityLibrary.h b/Engine/Source/Developer/FunctionalTesting/Classes/FunctionalTestUtilityLibrary.h
new file mode 100644
--- /dev/null
+++ b/Engine/Source/Developer/FunctionalTesting/Classes/FunctionalTestUtilityLibrary.h
@@ -0,0 +1,24 @@
+// Copyright 1998-2017 Epic Games, Inc. All Rights Reserved.
+
+#pragma once
+
+#include "Kismet/KismetSystemLibrary.h"
+#include "Kismet/BlueprintFunctionLibrary.h"
+#include "TraceQueryTestResults.h"
+#include "FunctionalTestUtilityLibrary.generated.h"
+
+/**
+* Used to expose C++ functions to tests that we don't want to make BP accessible
+* in the engine itself.
+*/
+UCLASS()
+class FUNCTIONALTESTING_API UFunctionalTestUtilityLibrary : public UBlueprintFunctionLibrary
+{
+	GENERATED_BODY()
+
+public:
+
+	/** Helper function to trace and permute many options at once */
+	UFUNCTION(BlueprintCallable, Category = "Utility|Collision", meta = (bIgnoreSelf = "true", WorldContext = "WorldContextObject", AutoCreateRefTerm = "ActorsToIgnore", AdvancedDisplay = "TraceColor,TraceHitColor,DrawTime", Keywords = "sweep"))
+	static UTraceQueryTestResults* TraceChannelTestUtil(UObject* WorldContextObject, const FTraceChannelTestBatchOptions& BatchOptions, const FVector Start, const FVector End, float SphereCapsuleRadius, float CapsuleHalfHeight, FVector BoxHalfSize, const FRotator Orientation, ETraceTypeQuery TraceChannel, TArray<TEnumAsByte<EObjectTypeQuery> > ObjectTypes, FName ProfileName, bool bTraceComplex, const TArray<AActor*>& ActorsToIgnore, bool bIgnoreSelf, EDrawDebugTrace::Type DrawDebugType, FLinearColor TraceColor = FLinearColor::Red, FLinearColor TraceHitColor = FLinearColor::Green, float DrawTime = 5.0f);
+};
diff --git a/Engine/Source/Developer/FunctionalTesting/Classes/FunctionalTestingManager.h b/Engine/Source/Developer/FunctionalTesting/Classes/FunctionalTestingManager.h
--- a/Engine/Source/Developer/FunctionalTesting/Classes/FunctionalTestingManager.h
+++ b/Engine/Source/Developer/FunctionalTesting/Classes/FunctionalTestingManager.h
@@ -76,7 +76,6 @@ protected:
 	bool bIsRunning;
 	bool bFinished;
 	bool bLooped;
-	bool bWaitForNavigationBuildFinish;
 	bool bInitialDelayApplied;
 	uint32 CurrentIteration;
 
diff --git a/Engine/Source/Developer/FunctionalTesting/Classes/TraceQueryTestResults.h b/Engine/Source/Developer/FunctionalTesting/Classes/TraceQueryTestResults.h
new file mode 100644
--- /dev/null
+++ b/Engine/Source/Developer/FunctionalTesting/Classes/TraceQueryTestResults.h
@@ -0,0 +1,181 @@
+// Copyright 1998-2017 Epic Games, Inc. All Rights Reserved.
+
+#pragma once
+
+#include "CoreMinimal.h"
+#include "UObject/ObjectMacros.h"
+#include "UObject/UObjectGlobals.h"
+#include "Engine/EngineTypes.h"
+#include "TraceQueryTestResults.generated.h"
+
+class AFunctionalTest;
+
+USTRUCT(BlueprintType)
+struct FUNCTIONALTESTING_API FTraceChannelTestBatchOptions
+{
+	GENERATED_BODY()
+
+	FTraceChannelTestBatchOptions()
+	{
+		bLineTrace = true;
+		bSphereTrace = false;
+		bCapsuleTrace = false;
+		bBoxTrace = false;
+
+		bChannelTrace = true;
+		bObjectsTrace = false;
+		bProfileTrace = false;
+	}
+
+	bool operator==(const FTraceChannelTestBatchOptions& Other) const
+	{
+		return bLineTrace == Other.bLineTrace && bSphereTrace == Other.bSphereTrace && bCapsuleTrace == Other.bCapsuleTrace && bBoxTrace == Other.bBoxTrace
+			&& bChannelTrace == Other.bChannelTrace && bObjectsTrace == Other.bObjectsTrace && bProfileTrace == Other.bProfileTrace;
+	}
+
+	bool operator!=(const FTraceChannelTestBatchOptions& Other) const
+	{
+		return !(*this == Other);
+	}
+
+	/** Whether to do line traces */
+	UPROPERTY(EditAnywhere, BlueprintReadWrite, Category = "Utility|Collision")
+	bool bLineTrace;
+
+	/** Whether to do sphere traces */
+	UPROPERTY(EditAnywhere, BlueprintReadWrite, Category = "Utility|Collision")
+	bool bSphereTrace;
+
+	/** Whether to do capsule traces */
+	UPROPERTY(EditAnywhere, BlueprintReadWrite, Category = "Utility|Collision")
+	bool bCapsuleTrace;
+
+	/** Whether to do box traces */
+	UPROPERTY(EditAnywhere, BlueprintReadWrite, Category = "Utility|Collision")
+	bool bBoxTrace;
+
+	/** Whether to do channel traces */
+	UPROPERTY(EditAnywhere, BlueprintReadWrite, Category = "Utility|Collision")
+	bool bChannelTrace;
+
+	/** Whether to do object traces */
+	UPROPERTY(EditAnywhere, BlueprintReadWrite, Category = "Utility|Collision")
+	bool bObjectsTrace;
+
+	/** Whether to do profile traces */
+	UPROPERTY(EditAnywhere, BlueprintReadWrite, Category = "Utility|Collision")
+	bool bProfileTrace;
+
+	FString ToString() const;
+};
+
+USTRUCT(BlueprintType)
+struct FUNCTIONALTESTING_API FTraceQueryTestNames
+{
+	GENERATED_BODY()
+
+	UPROPERTY(EditAnywhere, BlueprintReadWrite, Category = "Utility|Collision")
+	FName ComponentName;
+
+	UPROPERTY(EditAnywhere, BlueprintReadWrite, Category = "Utility|Collision")
+	FName PhysicalMaterialName;
+
+	UPROPERTY(EditAnywhere, BlueprintReadWrite, Category = "Utility|Collision")
+	FName ActorName;
+
+	FString ToString() const;
+};
+
+USTRUCT(BlueprintType)
+struct FUNCTIONALTESTING_API FTraceQueryTestResultsInnerMost
+{
+	GENERATED_BODY()
+
+	/** Result from doing a single sweep*/
+	UPROPERTY(EditAnywhere, BlueprintReadWrite, Category = "Utility|Collision")
+	FHitResult SingleHit;
+
+	/** Names found from doing a single sweep*/
+	UPROPERTY(EditAnywhere, BlueprintReadWrite, Category = "Utility|Collision")
+	FTraceQueryTestNames SingleNames;
+
+	/** The true/false value returned from the single sweep */
+	UPROPERTY(EditAnywhere, BlueprintReadWrite, Category = "Utility|Collision")
+	bool bSingleResult;
+
+	/** Result from doing a multi sweep */
+	UPROPERTY(EditAnywhere, BlueprintReadWrite, Category = "Utility|Collision")
+	TArray<FHitResult> MultiHits;
+
+	/** Names found from doing a multi sweep*/
+	UPROPERTY(EditAnywhere, BlueprintReadWrite, Category = "Utility|Collision")
+	TArray<FTraceQueryTestNames> MultiNames;
+
+	/** The true/false value returned from the multi sweep */
+	UPROPERTY(EditAnywhere, BlueprintReadWrite, Category = "Utility|Collision")
+	bool bMultiResult;
+
+	FString ToString() const;
+
+	void CaptureNames();
+
+	bool AssertEqual(const FTraceQueryTestResultsInnerMost& Expected, const FString& What, const UObject* ContextObject, AFunctionalTest& FunctionalTest) const;
+};
+
+USTRUCT(BlueprintType)
+struct FUNCTIONALTESTING_API FTraceQueryTestResultsInner
+{
+	GENERATED_BODY()
+
+	/** The results associated with the line trace */
+	UPROPERTY(EditAnywhere, BlueprintReadWrite, Category = "Utility|Collision")
+	FTraceQueryTestResultsInnerMost LineResults;
+
+	/** The results associated with the sphere */
+	UPROPERTY(EditAnywhere, BlueprintReadWrite, Category = "Utility|Collision")
+	FTraceQueryTestResultsInnerMost SphereResults;
+
+	/** The results associated with the capsule*/
+	UPROPERTY(EditAnywhere, BlueprintReadWrite, Category = "Utility|Collision")
+	FTraceQueryTestResultsInnerMost CapsuleResults;
+
+	/** The results associated with the box*/
+	UPROPERTY(EditAnywhere, BlueprintReadWrite, Category = "Utility|Collision")
+	FTraceQueryTestResultsInnerMost BoxResults;
+	
+	FString ToString(const FTraceChannelTestBatchOptions& BatchOptions) const;
+
+	void CaptureNames();
+
+	bool AssertEqual(const FTraceQueryTestResultsInner& Expected, const FString& What, const UObject* ContextObject, AFunctionalTest& FunctionalTest) const;
+};
+
+UCLASS(BlueprintType)
+class FUNCTIONALTESTING_API UTraceQueryTestResults : public UObject
+{
+	GENERATED_BODY()
+public:
+	/** Results for channel trace */
+	UPROPERTY(EditAnywhere, BlueprintReadWrite, Category = "Utility|Collision")
+	FTraceQueryTestResultsInner ChannelResults;
+
+	/** Results for object trace */
+	UPROPERTY(EditAnywhere, BlueprintReadWrite, Category = "Utility|Collision")
+	FTraceQueryTestResultsInner ObjectResults;
+
+	/** Results for profile trace */
+	UPROPERTY(EditAnywhere, BlueprintReadWrite, Category = "Utility|Collision")
+	FTraceQueryTestResultsInner ProfileResults;
+
+	FTraceChannelTestBatchOptions BatchOptions;
+
+	/** Output string value */
+	UFUNCTION(BlueprintCallable, Category = "Utility|Collision")
+	FString ToString();
+
+	bool AssertEqual(const UTraceQueryTestResults* Expected, const FString& What, const UObject* ContextObject, AFunctionalTest& FunctionalTest) const;
+
+	void CaptureNames();
+
+	
+};
\ No newline at end of file
diff --git a/Engine/Source/Developer/FunctionalTesting/FunctionalTesting.Build.cs b/Engine/Source/Developer/FunctionalTesting/FunctionalTesting.Build.cs
--- a/Engine/Source/Developer/FunctionalTesting/FunctionalTesting.Build.cs
+++ b/Engine/Source/Developer/FunctionalTesting/FunctionalTesting.Build.cs
@@ -4,31 +4,40 @@ using UnrealBuildTool;
 
 public class FunctionalTesting : ModuleRules
 {
-	public FunctionalTesting(ReadOnlyTargetRules Target) : base(Target)
-	{
-		PrivateDependencyModuleNames.AddRange(
-			new string[] {
-				"Core",
-				"CoreUObject",
-				"Engine",
+    public FunctionalTesting(ReadOnlyTargetRules Target) : base(Target)
+    {
+        PrivateDependencyModuleNames.AddRange(
+            new string[] {
+                "Core",
+                "CoreUObject",
+                "Engine",
                 "ShaderCore",
-				"Slate",
+                "Slate",
                 "MessageLog",
                 "AIModule",
                 "RenderCore",
                 "AssetRegistry",
                 "RHI",
-				"UMG"
-			}
-			);
+                "UMG"
+            }
+        );
 
-		PrivateIncludePaths.AddRange(
-			new string[] 
-			{
-				"MessageLog/Public",
-				"Stats/Public",
-				"Developer/FunctionalTesting/Private",
-			}
-		);
-	}
+        PrivateIncludePaths.AddRange(
+            new string[]
+            {
+                "MessageLog/Public",
+                "Stats/Public",
+                "Developer/FunctionalTesting/Private",
+            }
+        );
+
+        if (UEBuildConfiguration.bBuildEditor == true)
+        {
+            PrivateDependencyModuleNames.AddRange(
+                new string[] {
+                    "SourceControl"
+                }
+            );
+        }
+    }
 }
diff --git a/Engine/Source/Developer/FunctionalTesting/Private/AutomationBlueprintFunctionLibrary.cpp b/Engine/Source/Developer/FunctionalTesting/Private/AutomationBlueprintFunctionLibrary.cpp
--- a/Engine/Source/Developer/FunctionalTesting/Private/AutomationBlueprintFunctionLibrary.cpp
+++ b/Engine/Source/Developer/FunctionalTesting/Private/AutomationBlueprintFunctionLibrary.cpp
@@ -27,6 +27,7 @@
 #include "Engine/LocalPlayer.h"
 #include "ContentStreaming.h"
 #include "Stats/StatsData.h"
+#include "HAL/PlatformProperties.h"
 
 #define LOCTEXT_NAMESPACE "Automation"
 
@@ -48,8 +49,11 @@ static TAutoConsoleVariable<int32> CVarAutomationScreenshotResolutionHeight(
 
 void FinishLoadingBeforeScreenshot()
 {
-	// Force all shader compiling to finish.
-	GShaderCompilingManager->FinishAllCompilation();
+	// Finish compiling the shaders if the platform doesn't require cooked data.
+	if (!FPlatformProperties::RequiresCookedData())
+	{
+		GShaderCompilingManager->FinishAllCompilation();
+	}
 
 	// Force all mip maps to load before taking the screenshot.
 	UTexture::ForceUpdateTextureStreaming();
@@ -572,9 +576,10 @@ bool UAutomationBlueprintFunctionLibrary::AreAutomatedTestsRunning()
 	return GIsAutomationTesting;
 }
 
-FAutomationScreenshotOptions UAutomationBlueprintFunctionLibrary::GetDefaultScreenshotOptionsForGameplay(EComparisonTolerance Tolerance)
+FAutomationScreenshotOptions UAutomationBlueprintFunctionLibrary::GetDefaultScreenshotOptionsForGameplay(EComparisonTolerance Tolerance, float Delay)
 {
 	FAutomationScreenshotOptions Options;
+	Options.Delay = Delay;
 	Options.Tolerance = Tolerance;
 	Options.bDisableNoisyRenderingFeatures = true;
 	Options.bIgnoreAntiAliasing = true;
@@ -583,9 +588,10 @@ FAutomationScreenshotOptions UAutomationBlueprintFunctionLibrary::GetDefaultScre
 	return Options;
 }
 
-FAutomationScreenshotOptions UAutomationBlueprintFunctionLibrary::GetDefaultScreenshotOptionsForRendering(EComparisonTolerance Tolerance)
+FAutomationScreenshotOptions UAutomationBlueprintFunctionLibrary::GetDefaultScreenshotOptionsForRendering(EComparisonTolerance Tolerance, float Delay)
 {
 	FAutomationScreenshotOptions Options;
+	Options.Delay = Delay;
 	Options.Tolerance = Tolerance;
 	Options.bDisableNoisyRenderingFeatures = true;
 	Options.bIgnoreAntiAliasing = true;
diff --git a/Engine/Source/Developer/FunctionalTesting/Private/FuncTestManager.cpp b/Engine/Source/Developer/FunctionalTesting/Private/FuncTestManager.cpp
--- a/Engine/Source/Developer/FunctionalTesting/Private/FuncTestManager.cpp
+++ b/Engine/Source/Developer/FunctionalTesting/Private/FuncTestManager.cpp
@@ -56,6 +56,11 @@ void FFuncTestManager::SetScript(class UFunctionalTestingManager* NewScript)
 	TestScript = NewScript;
 }
 
+UFunctionalTestingManager* FFuncTestManager::GetCurrentScript()
+{
+	return TestScript.Get();
+}
+
 bool FFuncTestManager::IsRunning() const 
 { 
 	return TestScript.IsValid() && TestScript->IsRunning();
diff --git a/Engine/Source/Developer/FunctionalTesting/Private/FuncTestManager.h b/Engine/Source/Developer/FunctionalTesting/Private/FuncTestManager.h
--- a/Engine/Source/Developer/FunctionalTesting/Private/FuncTestManager.h
+++ b/Engine/Source/Developer/FunctionalTesting/Private/FuncTestManager.h
@@ -33,7 +33,7 @@ public:
 	
 	virtual void SetScript(class UFunctionalTestingManager* NewScript) override;
 
-	virtual class UFunctionalTestingManager* GetCurrentScript() override { return TestScript.Get(); }
+	virtual class UFunctionalTestingManager* GetCurrentScript();
 
 	virtual void SetLooping(const bool bLoop) override;
 
diff --git a/Engine/Source/Developer/FunctionalTesting/Private/FunctionalTest.cpp b/Engine/Source/Developer/FunctionalTesting/Private/FunctionalTest.cpp
--- a/Engine/Source/Developer/FunctionalTesting/Private/FunctionalTest.cpp
+++ b/Engine/Source/Developer/FunctionalTesting/Private/FunctionalTest.cpp
@@ -22,6 +22,7 @@
 #include "Engine/Texture2D.h"
 #include "DelayForFramesLatentAction.h"
 #include "Engine/DebugCameraController.h"
+#include "TraceQueryTestResults.h"
 
 namespace
 {
@@ -99,6 +100,7 @@ AFunctionalTest::AFunctionalTest( const FObjectInitializer& ObjectInitializer )
 {
 	PrimaryActorTick.bCanEverTick = true;
 	PrimaryActorTick.bStartWithTickEnabled = false;
+	PrimaryActorTick.bTickEvenWhenPaused = true;
 	
 	bCanBeDamaged = false;
 
@@ -526,194 +528,255 @@ UBillboardComponent* AFunctionalTest::GetSpriteComponent()
 
 //////////////////////////////////////////////////////////////////////////
 
-void AFunctionalTest::AssertTrue(bool Condition, FString Message, const UObject* ContextObject)
+bool AFunctionalTest::AssertTrue(bool Condition, FString Message, const UObject* ContextObject)
 {
 	if ( !Condition )
 	{
 		LogStep(ELogVerbosity::Error, FString::Printf(TEXT("Assertion failed: '%s' for context '%s'"), *Message, ContextObject ? *ContextObject->GetName() : TEXT("")));
+		return false;
 	}
 	else
 	{
 		LogStep(ELogVerbosity::Log, FString::Printf(TEXT("Assertion passed (%s)"), *Message));
+		return true;
 	}
 }
 
-void AFunctionalTest::AssertFalse(bool Condition, FString Message, const UObject* ContextObject)
+bool AFunctionalTest::AssertFalse(bool Condition, FString Message, const UObject* ContextObject)
 {
-	AssertTrue(!Condition, Message, ContextObject);
+	return AssertTrue(!Condition, Message, ContextObject);
 }
 
-void AFunctionalTest::AssertIsValid(UObject* Object, FString Message, const UObject* ContextObject)
+bool AFunctionalTest::AssertIsValid(UObject* Object, FString Message, const UObject* ContextObject)
 {
 	if ( !IsValid(Object) )
 	{
 		LogStep(ELogVerbosity::Error, FString::Printf(TEXT("Invalid object: '%s' for context '%s'"), *Message, ContextObject ? *ContextObject->GetName() : TEXT("")));
+		return false;
 	}
 	else
 	{
 		LogStep(ELogVerbosity::Log, FString::Printf(TEXT("Valid object: (%s)"), *Message));
+		return true;
 	}
 }
 
-void AFunctionalTest::AssertValue_Int(int32 Actual, EComparisonMethod ShouldBe, int32 Expected, const FString& What, const UObject* ContextObject)
+bool AFunctionalTest::AssertValue_Int(int32 Actual, EComparisonMethod ShouldBe, int32 Expected, const FString& What, const UObject* ContextObject)
 {
 	if ( !PerformComparison(Actual, Expected, ShouldBe) )
 	{
 		LogStep(ELogVerbosity::Error, FString::Printf(TEXT("%s: expected {%d} to be %s {%d} for context '%s'"), *What, Actual, *GetComparisonAsString(ShouldBe), Expected, ContextObject ? *ContextObject->GetName() : TEXT("")));
+		return false;
 	}
 	else
 	{
-		LogStep(ELogVerbosity::Log, FString::Printf(TEXT("Int assertion passed (%s)"), *What));
+		LogStep(ELogVerbosity::Log, FString::Printf(TEXT("%s: expected {%d} to be %s {%d} for context '%s'"), *What, Actual, *GetComparisonAsString(ShouldBe), Expected, ContextObject ? *ContextObject->GetName() : TEXT("")));
+		return true;
 	}
 }
 
-void AFunctionalTest::AssertValue_Float(float Actual, EComparisonMethod ShouldBe, float Expected, const FString& What, const UObject* ContextObject)
+bool AFunctionalTest::AssertValue_Float(float Actual, EComparisonMethod ShouldBe, float Expected, const FString& What, const UObject* ContextObject)
 {
 	if ( !PerformComparison(Actual, Expected, ShouldBe) )
 	{
 		LogStep(ELogVerbosity::Error, FString::Printf(TEXT("%s: expected {%f} to be %s {%f} for context '%s'"), *What, Actual, *GetComparisonAsString(ShouldBe), Expected, ContextObject ? *ContextObject->GetName() : TEXT("")));
+		return false;
 	}
 	else
 	{
-		LogStep(ELogVerbosity::Log, FString::Printf(TEXT("Float assertion passed (%s)"), *What));
+		LogStep(ELogVerbosity::Log, FString::Printf(TEXT("%s: expected {%f} to be %s {%f} for context '%s'"), *What, Actual, *GetComparisonAsString(ShouldBe), Expected, ContextObject ? *ContextObject->GetName() : TEXT("")));
+		return true;
 	}
 }
 
-void AFunctionalTest::AssertValue_DateTime(FDateTime Actual, EComparisonMethod ShouldBe, FDateTime Expected, const FString& What, const UObject* ContextObject)
+bool AFunctionalTest::AssertValue_DateTime(FDateTime Actual, EComparisonMethod ShouldBe, FDateTime Expected, const FString& What, const UObject* ContextObject)
 {
 	if ( !PerformComparison(Actual, Expected, ShouldBe) )
 	{
 		LogStep(ELogVerbosity::Error, FString::Printf(TEXT("%s: expected {%s} to be %s {%s} for context '%s'"), *What, *Actual.ToString(), *GetComparisonAsString(ShouldBe), *Expected.ToString(), ContextObject ? *ContextObject->GetName() : TEXT("")));
+		return false;
 	}
 	else
 	{
 		LogStep(ELogVerbosity::Log, FString::Printf(TEXT("DateTime assertion passed (%s)"), *What));
+		return true;
 	}
 }
 
-
-void AFunctionalTest::AssertEqual_Float(const float Actual, const float Expected, const FString& What, const float Tolerance, const UObject* ContextObject)
+bool AFunctionalTest::AssertEqual_Float(const float Actual, const float Expected, const FString& What, const float Tolerance, const UObject* ContextObject)
 {
 	if ( !FMath::IsNearlyEqual(Actual, Expected, Tolerance) )
 	{
 		LogStep(ELogVerbosity::Error, FString::Printf(TEXT("Expected '%s' to be {%f}, but it was {%f} within tolerance {%f} for context '%s'"), *What, Expected, Actual, Tolerance, ContextObject ? *ContextObject->GetName() : TEXT("")));
+		return false;
 	}
 	else
 	{
 		LogStep(ELogVerbosity::Log, FString::Printf(TEXT("Float assertion passed (%s)"), *What));
+		return true;
+	}
+}
+
+bool AFunctionalTest::AssertEqual_Bool(const bool Actual, const bool Expected, const FString& What, const UObject* ContextObject)
+{
+	if (Actual != Expected)
+	{
+		LogStep(ELogVerbosity::Error, FString::Printf(TEXT("Expected '%d' to be {%d}, but it was {%d} for context '%s'"), *What, Expected, Actual, ContextObject ? *ContextObject->GetName() : TEXT("")));
+		return false;
+	}
+	else
+	{
+		LogStep(ELogVerbosity::Log, FString::Printf(TEXT("Bool assertion passed (%s)"), *What));
+		return true;
 	}
 }
 
 bool AFunctionalTest::AssertEqual_Int(const int32 Actual, const int32 Expected, const FString& What, const UObject* ContextObject)
 {
-	if(Actual != Expected)
+	if (Actual != Expected)
+	{
+		LogStep(ELogVerbosity::Error, FString::Printf(TEXT("Expected '%d' to be {%d}, but it was {%d} for context '%s'"), *What, Expected, Actual, ContextObject ? *ContextObject->GetName() : TEXT("")));
+		return false;
+	}
+	else
+	{
+		LogStep(ELogVerbosity::Log, FString::Printf(TEXT("Bool assertion passed (%s)"), *What));
+		return true;
+	}
+}
+
+bool AFunctionalTest::AssertEqual_Name(const FName Actual, const FName Expected, const FString& What, const UObject* ContextObject)
+{
+	if (Actual != Expected)
 	{
-		LogStep(ELogVerbosity::Error, FString::Printf(TEXT("Expected '%s' to be {%d}, but it was {%d} for context '%s'"), *What, Expected, Actual, ContextObject ? *ContextObject->GetName() : TEXT("")));
+		LogStep(ELogVerbosity::Error, FString::Printf(TEXT("Expected '%s' to be {%s}, but it was {%s} for context '%s'"), *What, *Expected.ToString(), *Actual.ToString(), ContextObject ? *ContextObject->GetName() : TEXT("")));
 		return false;
 	}
 	else
 	{
-		LogStep(ELogVerbosity::Log, FString::Printf(TEXT("Int assertion passed (%s)"), *What));
+		LogStep(ELogVerbosity::Log, FString::Printf(TEXT("Bool assertion passed (%s)"), *What));
 		return true;
 	}
 }
 
-void AFunctionalTest::AssertEqual_Transform(const FTransform Actual, const FTransform Expected, const FString& What, const UObject* ContextObject)
+
+bool AFunctionalTest::AssertEqual_Transform(const FTransform Actual, const FTransform Expected, const FString& What, const UObject* ContextObject)
 {
 	if ( !Expected.Equals(Actual) )
 	{
 		LogStep(ELogVerbosity::Error, FString::Printf(TEXT("Expected '%s' to be {%s}, but it was {%s} for context '%s'"), *What, *TransformToString(Expected), *TransformToString(Actual), ContextObject ? *ContextObject->GetName() : TEXT("")));
+		return false;
 	}
 	else
 	{
 		LogStep(ELogVerbosity::Log, FString::Printf(TEXT("Transform assertion passed (%s)"), *What));
+		return true;
 	}
 }
 
-void AFunctionalTest::AssertNotEqual_Transform(const FTransform Actual, const FTransform NotExpected, const FString& What, const UObject* ContextObject)
+bool AFunctionalTest::AssertNotEqual_Transform(const FTransform Actual, const FTransform NotExpected, const FString& What, const UObject* ContextObject)
 {
 	if ( NotExpected.Equals(Actual) )
 	{
 		LogStep(ELogVerbosity::Error, FString::Printf(TEXT("Expected '%s' not to be {%s} for context '%s'"), *What, *TransformToString(NotExpected), ContextObject ? *ContextObject->GetName() : TEXT("")));
+		return false;
 	}
 	else
 	{
 		LogStep(ELogVerbosity::Log, FString::Printf(TEXT("Transform assertion passed (%s)"), *What));
+		return true;
 	}
 }
 
-void AFunctionalTest::AssertEqual_Rotator(const FRotator Actual, const FRotator Expected, const FString& What, const UObject* ContextObject)
+bool AFunctionalTest::AssertEqual_Rotator(const FRotator Actual, const FRotator Expected, const FString& What, const UObject* ContextObject)
 {
 	if ( !Expected.Equals(Actual) )
 	{
 		LogStep(ELogVerbosity::Error, FString::Printf(TEXT("Expected '%s' to be {%s} but it was {%s} for context '%s'"), *What, *Expected.ToString(), *Actual.ToString(), ContextObject ? *ContextObject->GetName() : TEXT("")));
+		return false;
 	}
 	else
 	{
 		LogStep(ELogVerbosity::Log, FString::Printf(TEXT("Rotator assertion passed (%s)"), *What));
+		return true;
 	}
 }
 
-void AFunctionalTest::AssertNotEqual_Rotator(const FRotator Actual, const FRotator NotExpected, const FString& What, const UObject* ContextObject)
+bool AFunctionalTest::AssertNotEqual_Rotator(const FRotator Actual, const FRotator NotExpected, const FString& What, const UObject* ContextObject)
 {
 	if ( NotExpected.Equals(Actual) )
 	{
 		LogStep(ELogVerbosity::Error, FString::Printf(TEXT("Expected '%s' not to be {%s} for context '%s'"), *What, *NotExpected.ToString(), ContextObject ? *ContextObject->GetName() : TEXT("")));
+		return false;
 	}
 	else
 	{
 		LogStep(ELogVerbosity::Log, FString::Printf(TEXT("Rotator assertion passed (%s)"), *What));
+		return true;
 	}
 }
 
-void AFunctionalTest::AssertEqual_Vector(const FVector Actual, const FVector Expected, const FString& What, const float Tolerance, const UObject* ContextObject)
+bool AFunctionalTest::AssertEqual_Vector(const FVector Actual, const FVector Expected, const FString& What, const float Tolerance, const UObject* ContextObject)
 {
 	if ( !Expected.Equals(Actual, Tolerance) )
 	{
 		LogStep(ELogVerbosity::Error, FString::Printf(TEXT("Expected '%s' to be {%s} but it was {%s} within tolerance {%f} for context '%s'"), *What, *Expected.ToString(), *Actual.ToString(), Tolerance, ContextObject ? *ContextObject->GetName() : TEXT("")));
+		return false;
 	}
 	else
 	{
 		LogStep(ELogVerbosity::Log, FString::Printf(TEXT("Vector assertion passed (%s)"), *What));
+		return true;
 	}
 }
 
-void AFunctionalTest::AssertNotEqual_Vector(const FVector Actual, const FVector NotExpected, const FString& What, const UObject* ContextObject)
+bool AFunctionalTest::AssertNotEqual_Vector(const FVector Actual, const FVector NotExpected, const FString& What, const UObject* ContextObject)
 {
 	if ( NotExpected.Equals(Actual) )
 	{
 		LogStep(ELogVerbosity::Error, FString::Printf(TEXT("Expected '%s' not to be {%s} for context '%s'"), *What, *NotExpected.ToString(), ContextObject ? *ContextObject->GetName() : TEXT("")));
+		return false;
 	}
 	else
 	{
 		LogStep(ELogVerbosity::Log, FString::Printf(TEXT("Vector assertion passed (%s)"), *What));
+		return true;
 	}
 }
 
-void AFunctionalTest::AssertEqual_String(const FString Actual, const FString Expected, const FString& What, const UObject* ContextObject)
+bool AFunctionalTest::AssertEqual_String(const FString Actual, const FString Expected, const FString& What, const UObject* ContextObject)
 {
 	if ( !Expected.Equals(Actual) )
 	{
 		LogStep(ELogVerbosity::Error, FString::Printf(TEXT("Expected '%s' to be {%s} but it was {%s} for context '%s'"), *What, *Expected, *Actual, ContextObject ? *ContextObject->GetName() : TEXT("")));
+		return false;
 	}
 	else
 	{
 		LogStep(ELogVerbosity::Log, FString::Printf(TEXT("String assertion passed (%s)"), *What));
+		return true;
 	}
 }
 
-void AFunctionalTest::AssertNotEqual_String(const FString Actual, const FString NotExpected, const FString& What, const UObject* ContextObject)
+bool AFunctionalTest::AssertNotEqual_String(const FString Actual, const FString NotExpected, const FString& What, const UObject* ContextObject)
 {
 	if ( NotExpected.Equals(Actual) )
 	{
 		LogStep(ELogVerbosity::Error, FString::Printf(TEXT("Expected '%s' not to be {%s} for context '%s'"), *What, *NotExpected, ContextObject ? *ContextObject->GetName() : TEXT("")));
+		return false;
 	}
 	else
 	{
 		LogStep(ELogVerbosity::Log, FString::Printf(TEXT("String assertion passed (%s)"), *What));
+		return true;
 	}
 }
 
+bool AFunctionalTest::AssertEqual_TraceQueryResults(const UTraceQueryTestResults* Actual, const UTraceQueryTestResults* Expected, const FString& What, const UObject* ContextObject)
+{
+	return Actual->AssertEqual(Expected, What, ContextObject, *this);
+}
+
 void AFunctionalTest::AddWarning(const FString Message)
 {
 	LogStep(ELogVerbosity::Warning, Message);
@@ -740,9 +803,10 @@ void AFunctionalTest::LogStep(ELogVerbosity::Type Verbosity, const FString& Mess
 
 	switch ( Verbosity )
 	{
+		case ELogVerbosity::Display:
 		case ELogVerbosity::Log:
-			UE_VLOG(this, LogFunctionalTest, Log, TEXT("%s"), *FullMessage);
-			UE_LOG(LogFunctionalTest, Log, TEXT("%s"), *FullMessage);
+			UE_VLOG(this, LogFunctionalTest, Display, TEXT("%s"), *FullMessage);
+			UE_LOG(LogFunctionalTest, Display, TEXT("%s"), *FullMessage);
 		break;
 		case ELogVerbosity::Warning:
 			UE_VLOG(this, LogFunctionalTest, Warning, TEXT("%s"), *FullMessage);
diff --git a/Engine/Source/Developer/FunctionalTesting/Private/FunctionalTestUtilityLibrary.cpp b/Engine/Source/Developer/FunctionalTesting/Private/FunctionalTestUtilityLibrary.cpp
new file mode 100644
--- /dev/null
+++ b/Engine/Source/Developer/FunctionalTesting/Private/FunctionalTestUtilityLibrary.cpp
@@ -0,0 +1,105 @@
+// Copyright 1998-2017 Epic Games, Inc. All Rights Reserved.
+
+#include "FunctionalTestUtilityLibrary.h"
+#include "TraceQueryTestResults.h"
+
+UTraceQueryTestResults* UFunctionalTestUtilityLibrary::TraceChannelTestUtil(UObject* WorldContextObject, const FTraceChannelTestBatchOptions& BatchOptions, const FVector Start, const FVector End, float SphereCapsuleRadius, float CapsuleHalfHeight, FVector BoxHalfSize, const FRotator Orientation, ETraceTypeQuery TraceChannel, TArray<TEnumAsByte<EObjectTypeQuery> > ObjectTypes, FName ProfileName, bool bTraceComplex, const TArray<AActor*>& ActorsToIgnore, bool bIgnoreSelf, EDrawDebugTrace::Type DrawDebugType, FLinearColor TraceColor, FLinearColor TraceHitColor, float DrawTime)
+{
+	UTraceQueryTestResults* Results = NewObject<UTraceQueryTestResults>();
+	if(Results)
+	{
+		//line
+		Results->BatchOptions = BatchOptions;
+		if(BatchOptions.bLineTrace)
+		{
+			if(BatchOptions.bChannelTrace)
+			{
+				Results->ChannelResults.LineResults.bSingleResult = UKismetSystemLibrary::LineTraceSingle(WorldContextObject, Start, End, TraceChannel, bTraceComplex, ActorsToIgnore, DrawDebugType, Results->ChannelResults.LineResults.SingleHit, bIgnoreSelf, TraceColor, TraceHitColor, DrawTime);
+				Results->ChannelResults.LineResults.bMultiResult = UKismetSystemLibrary::LineTraceMulti(WorldContextObject, Start, End, TraceChannel, bTraceComplex, ActorsToIgnore, DrawDebugType, Results->ChannelResults.LineResults.MultiHits, bIgnoreSelf, TraceColor, TraceHitColor, DrawTime);
+			}
+			
+			if(BatchOptions.bObjectsTrace)
+			{
+				Results->ObjectResults.LineResults.bSingleResult = UKismetSystemLibrary::LineTraceSingleForObjects(WorldContextObject, Start, End, ObjectTypes, bTraceComplex, ActorsToIgnore, DrawDebugType, Results->ObjectResults.LineResults.SingleHit, bIgnoreSelf, TraceColor, TraceHitColor, DrawTime);
+				Results->ObjectResults.LineResults.bMultiResult = UKismetSystemLibrary::LineTraceMultiForObjects(WorldContextObject, Start, End, ObjectTypes, bTraceComplex, ActorsToIgnore, DrawDebugType, Results->ObjectResults.LineResults.MultiHits, bIgnoreSelf, TraceColor, TraceHitColor, DrawTime);
+			}
+			
+			if(BatchOptions.bProfileTrace)
+			{
+				Results->ProfileResults.LineResults.bSingleResult = UKismetSystemLibrary::LineTraceSingleByProfile(WorldContextObject, Start, End, ProfileName, bTraceComplex, ActorsToIgnore, DrawDebugType, Results->ProfileResults.LineResults.SingleHit, bIgnoreSelf, TraceColor, TraceHitColor, DrawTime);
+				Results->ProfileResults.LineResults.bMultiResult = UKismetSystemLibrary::LineTraceMultiByProfile(WorldContextObject, Start, End, ProfileName, bTraceComplex, ActorsToIgnore, DrawDebugType, Results->ProfileResults.LineResults.MultiHits, bIgnoreSelf, TraceColor, TraceHitColor, DrawTime);
+			}
+		}
+		
+
+		//sphere
+		if(BatchOptions.bSphereTrace)
+		{
+			if (BatchOptions.bChannelTrace)
+			{
+				Results->ChannelResults.SphereResults.bSingleResult = UKismetSystemLibrary::SphereTraceSingle(WorldContextObject, Start, End, SphereCapsuleRadius, TraceChannel, bTraceComplex, ActorsToIgnore, DrawDebugType, Results->ChannelResults.SphereResults.SingleHit, bIgnoreSelf, TraceColor, TraceHitColor, DrawTime);
+				Results->ChannelResults.SphereResults.bMultiResult = UKismetSystemLibrary::SphereTraceMulti(WorldContextObject, Start, End, SphereCapsuleRadius, TraceChannel, bTraceComplex, ActorsToIgnore, DrawDebugType, Results->ChannelResults.SphereResults.MultiHits, bIgnoreSelf, TraceColor, TraceHitColor, DrawTime);
+			}
+
+			if (BatchOptions.bObjectsTrace)
+			{
+				Results->ObjectResults.SphereResults.bSingleResult = UKismetSystemLibrary::SphereTraceSingleForObjects(WorldContextObject, Start, End, SphereCapsuleRadius, ObjectTypes, bTraceComplex, ActorsToIgnore, DrawDebugType, Results->ObjectResults.SphereResults.SingleHit, bIgnoreSelf, TraceColor, TraceHitColor, DrawTime);
+				Results->ObjectResults.SphereResults.bMultiResult = UKismetSystemLibrary::SphereTraceMultiForObjects(WorldContextObject, Start, End, SphereCapsuleRadius, ObjectTypes, bTraceComplex, ActorsToIgnore, DrawDebugType, Results->ObjectResults.SphereResults.MultiHits, bIgnoreSelf, TraceColor, TraceHitColor, DrawTime);
+			}
+
+			if (BatchOptions.bProfileTrace)
+			{
+				Results->ProfileResults.SphereResults.bSingleResult = UKismetSystemLibrary::SphereTraceSingleByProfile(WorldContextObject, Start, End, SphereCapsuleRadius, ProfileName, bTraceComplex, ActorsToIgnore, DrawDebugType, Results->ProfileResults.SphereResults.SingleHit, bIgnoreSelf, TraceColor, TraceHitColor, DrawTime);
+				Results->ProfileResults.SphereResults.bMultiResult = UKismetSystemLibrary::SphereTraceMultiByProfile(WorldContextObject, Start, End, SphereCapsuleRadius, ProfileName, bTraceComplex, ActorsToIgnore, DrawDebugType, Results->ProfileResults.SphereResults.MultiHits, bIgnoreSelf, TraceColor, TraceHitColor, DrawTime);
+			}
+		}
+		
+		//capsule
+		if(BatchOptions.bCapsuleTrace)
+		{
+			if (BatchOptions.bChannelTrace)
+			{
+				Results->ChannelResults.CapsuleResults.bSingleResult = UKismetSystemLibrary::CapsuleTraceSingle(WorldContextObject, Start, End, SphereCapsuleRadius, CapsuleHalfHeight, TraceChannel, bTraceComplex, ActorsToIgnore, DrawDebugType, Results->ChannelResults.CapsuleResults.SingleHit, bIgnoreSelf, TraceColor, TraceHitColor, DrawTime);
+				Results->ChannelResults.CapsuleResults.bMultiResult = UKismetSystemLibrary::CapsuleTraceMulti(WorldContextObject, Start, End, SphereCapsuleRadius, CapsuleHalfHeight, TraceChannel, bTraceComplex, ActorsToIgnore, DrawDebugType, Results->ChannelResults.CapsuleResults.MultiHits, bIgnoreSelf, TraceColor, TraceHitColor, DrawTime);
+			}
+			
+			if (BatchOptions.bObjectsTrace)
+			{
+				Results->ObjectResults.CapsuleResults.bSingleResult = UKismetSystemLibrary::CapsuleTraceSingleForObjects(WorldContextObject, Start, End, SphereCapsuleRadius, CapsuleHalfHeight, ObjectTypes, bTraceComplex, ActorsToIgnore, DrawDebugType, Results->ObjectResults.CapsuleResults.SingleHit, bIgnoreSelf, TraceColor, TraceHitColor, DrawTime);
+				Results->ObjectResults.CapsuleResults.bMultiResult = UKismetSystemLibrary::CapsuleTraceMultiForObjects(WorldContextObject, Start, End, SphereCapsuleRadius, CapsuleHalfHeight, ObjectTypes, bTraceComplex, ActorsToIgnore, DrawDebugType, Results->ObjectResults.CapsuleResults.MultiHits, bIgnoreSelf, TraceColor, TraceHitColor, DrawTime);
+			}
+
+			if (BatchOptions.bProfileTrace)
+			{
+				Results->ProfileResults.CapsuleResults.bSingleResult = UKismetSystemLibrary::CapsuleTraceSingleByProfile(WorldContextObject, Start, End, SphereCapsuleRadius, CapsuleHalfHeight, ProfileName, bTraceComplex, ActorsToIgnore, DrawDebugType, Results->ProfileResults.CapsuleResults.SingleHit, bIgnoreSelf, TraceColor, TraceHitColor, DrawTime);
+				Results->ProfileResults.CapsuleResults.bMultiResult = UKismetSystemLibrary::CapsuleTraceMultiByProfile(WorldContextObject, Start, End, SphereCapsuleRadius, CapsuleHalfHeight, ProfileName, bTraceComplex, ActorsToIgnore, DrawDebugType, Results->ProfileResults.CapsuleResults.MultiHits, bIgnoreSelf, TraceColor, TraceHitColor, DrawTime);
+			}
+		}
+
+		//box
+		if(BatchOptions.bBoxTrace)
+		{
+			if (BatchOptions.bChannelTrace)
+			{
+				Results->ChannelResults.BoxResults.bSingleResult = UKismetSystemLibrary::BoxTraceSingle(WorldContextObject, Start, End, BoxHalfSize, Orientation, TraceChannel, bTraceComplex, ActorsToIgnore, DrawDebugType, Results->ChannelResults.BoxResults.SingleHit, bIgnoreSelf, TraceColor, TraceHitColor, DrawTime);
+				Results->ChannelResults.BoxResults.bMultiResult = UKismetSystemLibrary::BoxTraceMulti(WorldContextObject, Start, End, BoxHalfSize, Orientation, TraceChannel, bTraceComplex, ActorsToIgnore, DrawDebugType, Results->ChannelResults.BoxResults.MultiHits, bIgnoreSelf, TraceColor, TraceHitColor, DrawTime);
+			}
+			
+			if (BatchOptions.bObjectsTrace)
+			{
+				Results->ObjectResults.BoxResults.bSingleResult = UKismetSystemLibrary::BoxTraceSingleForObjects(WorldContextObject, Start, End, BoxHalfSize, Orientation, ObjectTypes, bTraceComplex, ActorsToIgnore, DrawDebugType, Results->ObjectResults.BoxResults.SingleHit, bIgnoreSelf, TraceColor, TraceHitColor, DrawTime);
+				Results->ObjectResults.BoxResults.bMultiResult = UKismetSystemLibrary::BoxTraceMultiForObjects(WorldContextObject, Start, End, BoxHalfSize, Orientation, ObjectTypes, bTraceComplex, ActorsToIgnore, DrawDebugType, Results->ObjectResults.BoxResults.MultiHits, bIgnoreSelf, TraceColor, TraceHitColor, DrawTime);
+			}
+
+			if (BatchOptions.bProfileTrace)
+			{
+				Results->ProfileResults.BoxResults.bSingleResult = UKismetSystemLibrary::BoxTraceSingleByProfile(WorldContextObject, Start, End, BoxHalfSize, Orientation, ProfileName, bTraceComplex, ActorsToIgnore, DrawDebugType, Results->ProfileResults.BoxResults.SingleHit, bIgnoreSelf, TraceColor, TraceHitColor, DrawTime);
+				Results->ProfileResults.BoxResults.bMultiResult = UKismetSystemLibrary::BoxTraceMultiByProfile(WorldContextObject, Start, End, BoxHalfSize, Orientation, ProfileName, bTraceComplex, ActorsToIgnore, DrawDebugType, Results->ProfileResults.BoxResults.MultiHits, bIgnoreSelf, TraceColor, TraceHitColor, DrawTime);
+			}
+		}
+
+		Results->CaptureNames();
+	}
+
+	return Results;
+}
\ No newline at end of file
diff --git a/Engine/Source/Developer/FunctionalTesting/Private/FunctionalTestingManager.cpp b/Engine/Source/Developer/FunctionalTesting/Private/FunctionalTestingManager.cpp
--- a/Engine/Source/Developer/FunctionalTesting/Private/FunctionalTestingManager.cpp
+++ b/Engine/Source/Developer/FunctionalTesting/Private/FunctionalTestingManager.cpp
@@ -62,7 +62,6 @@ UFunctionalTestingManager::UFunctionalTestingManager( const FObjectInitializer&
 	, bIsRunning(false)
 	, bFinished(false)
 	, bLooped(false)
-	, bWaitForNavigationBuildFinish(false)
 	, bInitialDelayApplied(false)
 	, CurrentIteration(INDEX_NONE)
 {
@@ -100,7 +99,6 @@ bool UFunctionalTestingManager::RunAllFunctionalTests(UObject* WorldContextObjec
 
 	Manager->bFinished = false;
 	Manager->bLooped = bRunLooped;
-	Manager->bWaitForNavigationBuildFinish = bInWaitForNavigationBuildFinish;
 	Manager->CurrentIteration = 0;
 	Manager->TestsLeft.Reset();
 	Manager->AllTests.Reset();
@@ -158,7 +156,11 @@ void UFunctionalTestingManager::TriggerFirstValidTest()
 	check(World);
 	bIsRunning = World->GetNavigationSystem() != nullptr;
 
-	if (bInitialDelayApplied == true && (bWaitForNavigationBuildFinish == false || UNavigationSystem::IsNavigationBeingBuilt(World) == false) && World->AreActorsInitialized())
+	const bool bIsWorldInitialized =
+		World->AreActorsInitialized() &&
+		!UNavigationSystem::IsNavigationBeingBuilt(World);
+
+	if (bInitialDelayApplied == true && bIsWorldInitialized)
 	{
 		bIsRunning = RunFirstValidTest();
 		if (bIsRunning == false)
@@ -178,10 +180,11 @@ UFunctionalTestingManager* UFunctionalTestingManager::GetManager(UObject* WorldC
 {
 	UFunctionalTestingManager* Manager = FFunctionalTestingModule::Get()->GetCurrentScript();
 
-	if (Manager == NULL)
+	if (Manager == nullptr)
 	{
-		UObject* Outer = WorldContext ? WorldContext : (UObject*)GetTransientPackage();
-		Manager = NewObject<UFunctionalTestingManager>(Outer);
+		check(WorldContext);
+		UObject* World = GEngine->GetWorldFromContextObject(WorldContext);
+		Manager = NewObject<UFunctionalTestingManager>(World);
 		FFunctionalTestingModule::Get()->SetScript(Manager);
 
 		// add to root and get notified on world cleanup to remove from root on map cleanup
@@ -203,6 +206,9 @@ void UFunctionalTestingManager::OnWorldCleanedUp(UWorld* World, bool bSessionEnd
 	if (MyWorld == World)
 	{
 		RemoveFromRoot();
+
+		// Clear the functional test manager once the world is removed.
+		FFunctionalTestingModule::Get()->SetScript(nullptr);
 	}
 }
 
@@ -283,7 +289,12 @@ bool UFunctionalTestingManager::RunFirstValidTest()
 	if (TestReproStrings.Num() > 0)
 	{
 		UWorld* World = GetWorld();
-		UObject* TestsOuter = World ? (UObject*)(World->PersistentLevel) : (UObject*)(ANY_PACKAGE);
+
+		if (World == nullptr)
+		{
+			UE_LOG(LogFunctionalTest, Warning, TEXT("Unable to find testing world!"));
+			return bTestSuccessfullyTriggered;
+		}
 
 		while (TestReproStrings.Num() > 0)
 		{
@@ -302,7 +313,16 @@ bool UFunctionalTestingManager::RunFirstValidTest()
 			// first param is the test name. Look for it		
 			const FString TestName = TestParams[0];
 			TestParams.RemoveAt(0, 1, /*bAllowShrinking=*/false);
-			AFunctionalTest* TestToRun = FindObject<AFunctionalTest>(TestsOuter, *TestName);			
+
+			AFunctionalTest* TestToRun = nullptr;
+			for (TActorIterator<AFunctionalTest> It(World); It; ++It)
+			{
+				if (It->GetName() == TestName)
+				{
+					TestToRun = *It;
+				}
+			}
+			
 			if (TestToRun)
 			{
 				// Add the test we found to the tests left to run, so that if re-runs occur we continue to process this test until
@@ -322,7 +342,13 @@ bool UFunctionalTestingManager::RunFirstValidTest()
 			}
 			else
 			{
-				UE_LOG(LogFunctionalTest, Warning, TEXT("Unable to find test \'%s\'"), *TestName);
+				UE_LOG(LogFunctionalTest, Warning, TEXT("Unable to find test \'%s\' in world %s, the available tests are..."), *TestName, *World->GetFullName());
+
+				// Find Actors by type (needs a UWorld object)
+				for (TActorIterator<AFunctionalTest> It(World); It; ++It)
+				{
+					UE_LOG(LogFunctionalTest, Warning, TEXT("\'%s\'."), *It->GetName());
+				}
 			}
 		}
 	}
diff --git a/Engine/Source/Developer/FunctionalTesting/Private/GroundTruthData.cpp b/Engine/Source/Developer/FunctionalTesting/Private/GroundTruthData.cpp
new file mode 100644
--- /dev/null
+++ b/Engine/Source/Developer/FunctionalTesting/Private/GroundTruthData.cpp
@@ -0,0 +1,103 @@
+// Copyright 1998-2017 Epic Games, Inc. All Rights Reserved.
+
+#include "GroundTruthData.h"
+
+#if WITH_EDITOR
+#include "ISourceControlModule.h"
+#include "ISourceControlProvider.h"
+#include "ISourceControlOperation.h"
+#include "UObject/Package.h"
+#include "Misc/PackageName.h"
+#include "AssetData.h"
+#include "UObject/UnrealType.h"
+#endif
+
+DEFINE_LOG_CATEGORY_STATIC(GroundTruthLog, Log, Log)
+
+UGroundTruthData::UGroundTruthData()
+	: bResetGroundTruth(false)
+{
+}
+
+bool UGroundTruthData::CanModify() const
+{
+	return ObjectData == nullptr;
+}
+
+UObject* UGroundTruthData::LoadObject()
+{
+	UE_LOG(GroundTruthLog, Log, TEXT("Loaded Ground Truth, '%s'."), *GetPathName());
+	
+	return ObjectData;
+}
+
+void UGroundTruthData::SaveObject(UObject* GroundTruth)
+{
+	FAssetData GroundTruthAssetData(this);
+
+	UPackage* GroundTruthPackage = GetOutermost();
+	FString GroundTruthPackageName = GroundTruthAssetData.PackageName.ToString();
+
+#if WITH_EDITOR
+	if (!CanModify())
+	{
+		UE_LOG(GroundTruthLog, Warning, TEXT("Ground Truth, '%s' already set, unable to save changes.  Open and use bResetGroundTruth to reset the ground truth."), *GroundTruthPackageName);
+		return;
+	}
+
+	if (GroundTruth == nullptr)
+	{
+		UE_LOG(GroundTruthLog, Error, TEXT("Ground Truth, '%s' can not store a null object."), *GroundTruthPackageName);
+		return;
+	}
+
+	if (GIsBuildMachine)
+	{
+		UE_LOG(GroundTruthLog, Error, TEXT("Ground Truth, '%s' can not be modified on the build machine."), *GroundTruthPackageName);
+		return;
+	}
+
+	if (ISourceControlModule::Get().IsEnabled())
+	{
+		ISourceControlProvider& SourceControlProvider = ISourceControlModule::Get().GetProvider();
+		SourceControlProvider.Execute(ISourceControlOperation::Create<FMarkForAdd>(), GroundTruthPackage);
+		SourceControlProvider.Execute(ISourceControlOperation::Create<FCheckOut>(), GroundTruthPackage);
+	}
+
+	ObjectData = GroundTruth;
+	GroundTruth->Rename(nullptr, this);
+	MarkPackageDirty();
+
+	if (!UPackage::SavePackage(GroundTruthPackage, NULL, RF_Standalone, *FPackageName::LongPackageNameToFilename(GroundTruthPackageName, FPackageName::GetAssetPackageExtension()), GError, nullptr, false, true, SAVE_NoError))
+	{
+		UE_LOG(GroundTruthLog, Error, TEXT("Failed to save ground truth data! %s"), *GroundTruthPackageName);
+	}
+
+	UE_LOG(GroundTruthLog, Log, TEXT("Saved Ground Truth, '%s'."), *GroundTruthPackageName);
+
+#else
+	UE_LOG(GroundTruthLog, Error, TEXT("Can't save ground truth data outside of the editor, '%s'."), *GroundTruthPackageName);
+#endif
+}
+
+#if WITH_EDITOR
+
+void UGroundTruthData::PostEditChangeProperty(FPropertyChangedEvent& PropertyChangedEvent)
+{
+	Super::PostEditChangeProperty(PropertyChangedEvent);
+
+	if (PropertyChangedEvent.GetPropertyName() == GET_MEMBER_NAME_CHECKED(UGroundTruthData, bResetGroundTruth))
+	{
+		bResetGroundTruth = false;
+
+		if (ObjectData)
+		{
+			ObjectData->Rename(nullptr, GetTransientPackage());
+			ObjectData = nullptr;
+		}
+
+		MarkPackageDirty();
+	}
+}
+
+#endif
\ No newline at end of file
diff --git a/Engine/Source/Developer/FunctionalTesting/Private/TraceQueryTestResults.cpp b/Engine/Source/Developer/FunctionalTesting/Private/TraceQueryTestResults.cpp
new file mode 100644
--- /dev/null
+++ b/Engine/Source/Developer/FunctionalTesting/Private/TraceQueryTestResults.cpp
@@ -0,0 +1,226 @@
+// Copyright 1998-2017 Epic Games, Inc. All Rights Reserved.
+
+#include "TraceQueryTestResults.h"
+#include "FunctionalTest.h"
+#include "Components/PrimitiveComponent.h"
+#include "PhysicalMaterials/PhysicalMaterial.h"
+
+FString FTraceQueryTestNames::ToString() const
+{
+	return FString::Printf(TEXT("Component:%s, Actor:%s, PhysicalMaterial:%s"), *ComponentName.ToString(), *ActorName.ToString(), *PhysicalMaterialName.ToString());
+}
+
+FString FTraceQueryTestResultsInnerMost::ToString() const
+{
+	FString OutSingle = FString::Printf(TEXT("----SingleResult:%d %s {%s}"), bSingleResult, *SingleNames.ToString(), *SingleHit.ToString());
+	FString OutMulti = FString::Printf(TEXT("\n----MultiResult:%d"), bMultiResult);
+	
+	for(int32 Counter = 0; Counter < MultiHits.Num(); ++Counter)
+	{
+		OutMulti.Append(FString::Printf(TEXT("\n----[%d] %s {%s}"), Counter++, *MultiNames[Counter].ToString(),  *MultiHits[Counter].ToString()));
+	}
+
+	return OutSingle.Append(OutMulti);
+}
+
+FString FTraceQueryTestResultsInner::ToString(const FTraceChannelTestBatchOptions& BatchOptions) const
+{
+	FString LineResultsStr = BatchOptions.bLineTrace ? FString::Printf(TEXT("--LineResults:\n%s\n\n"), *LineResults.ToString()) : "";
+	FString SphereResultsStr = BatchOptions.bSphereTrace ? FString::Printf(TEXT("--SphereResults:\n%s\n\n"), *SphereResults.ToString()) : "";
+	FString CapsuleResultsStr = BatchOptions.bCapsuleTrace ? FString::Printf(TEXT("--CapsuleResults:\n%s\n\n"), *CapsuleResults.ToString()) : "";
+	FString BoxResultsStr = BatchOptions.bBoxTrace ? FString::Printf(TEXT("--BoxResults:\n%s\n\n"), *BoxResults.ToString()) : "";
+
+	return FString::Printf(TEXT("%s%s%s%s"), *LineResultsStr, *SphereResultsStr, *CapsuleResultsStr, *BoxResultsStr);
+}
+
+FString UTraceQueryTestResults::ToString()
+{
+	CaptureNames();
+
+	FString ChannelResultsStr = BatchOptions.bChannelTrace ? FString::Printf(TEXT("ChannelResults:\n%s\n\n"), *ChannelResults.ToString(BatchOptions)) : "";
+	FString ObjectResultsStr = BatchOptions.bObjectsTrace ? FString::Printf(TEXT("ObjectResults:\n%s\n\n"), *ObjectResults.ToString(BatchOptions)) : "";
+	FString ProfileResultsStr = BatchOptions.bProfileTrace ? FString::Printf(TEXT("ProfileResults:\n%s\n\n"), *ProfileResults.ToString(BatchOptions)) : "";
+	
+	return FString::Printf(TEXT("%s%s%s"), *ChannelResultsStr, *ObjectResultsStr, *ProfileResultsStr);
+}
+
+void CaptureNameHelper(FTraceQueryTestNames& Names, FHitResult& HitResult)
+{
+	UPrimitiveComponent* HitComp = HitResult.GetComponent();
+	AActor* HitActor = HitResult.GetActor();
+	UPhysicalMaterial* PhysMat = HitResult.PhysMaterial.Get();
+
+	Names.ComponentName = HitComp ? HitComp->GetFName() : NAME_None;
+	Names.ActorName = HitActor ? HitActor->GetFName() : NAME_None;
+	Names.PhysicalMaterialName = PhysMat ? PhysMat->GetFName() : NAME_None;
+}
+
+void FTraceQueryTestResultsInnerMost::CaptureNames()
+{
+	CaptureNameHelper(SingleNames, SingleHit);
+
+	MultiNames.Reset();
+	MultiNames.AddZeroed(MultiHits.Num());
+
+	for(int32 HitIdx = 0; HitIdx < MultiHits.Num(); ++HitIdx)
+	{
+		CaptureNameHelper(MultiNames[HitIdx], MultiHits[HitIdx]);
+	}
+}
+
+void FTraceQueryTestResultsInner::CaptureNames()
+{
+	LineResults.CaptureNames();
+	SphereResults.CaptureNames();
+	CapsuleResults.CaptureNames();
+	BoxResults.CaptureNames();
+}
+
+void UTraceQueryTestResults::CaptureNames()
+{
+	ChannelResults.CaptureNames();
+	ObjectResults.CaptureNames();
+	ProfileResults.CaptureNames();
+}
+
+FString FTraceChannelTestBatchOptions::ToString() const
+{
+	return FString::Printf(TEXT("bLineTrace:%d, bSphereTrace:%d, bCapsuleTrace:%d, bBoxTrace:%d, bChannelTrace:%d, bObjectsTrace:%d, bProfileTrace:%d"), bLineTrace, bSphereTrace, bCapsuleTrace, bBoxTrace, bChannelTrace, bObjectsTrace, bProfileTrace);
+}
+
+#define TEST_IMPL(X,Type) FunctionalTest.AssertEqual_##Type(Actual.##X, Expected.##X, FString::Printf(TEXT("%s:"#X), *What), ContextObject)
+#define TEST_IMPL_TOLERANCE(X, Type) FunctionalTest.AssertEqual_##Type(Actual.##X, Expected.##X, FString::Printf(TEXT("%s:"#X), *What), KINDA_SMALL_NUMBER, ContextObject)
+#define TEST_BOOL(X) TEST_IMPL(X, Bool)
+#define TEST_INT(X) TEST_IMPL(X, Int)
+#define TEST_NAME(X) TEST_IMPL(X, Name)
+#define TEST_FLOAT(X) TEST_IMPL_TOLERANCE(X, Float)
+#define TEST_VECTOR(X) TEST_IMPL_TOLERANCE(X, Vector)
+
+bool HelperAssertNamesEqual(const FTraceQueryTestNames& Actual, const FTraceQueryTestNames& Expected, const FString& What, const UObject* ContextObject, AFunctionalTest& FunctionalTest)
+{
+	bool bResults[] = {
+		TEST_NAME(ComponentName),
+		TEST_NAME(ActorName),
+		TEST_NAME(PhysicalMaterialName)
+	};
+	
+	for (bool bSuccess : bResults)
+	{
+		if (!bSuccess)
+		{
+			return false;
+		}
+	}
+
+	return true;
+}
+
+bool HelperAssertFHitResultEqual(const FHitResult& Actual, const FHitResult& Expected, const FString& What, const UObject* ContextObject, AFunctionalTest& FunctionalTest)
+{
+	bool bResults[] = {
+		TEST_BOOL(bBlockingHit),
+		TEST_BOOL(bStartPenetrating),
+		TEST_FLOAT(Time),
+		TEST_FLOAT(Distance),
+		TEST_VECTOR(Location),
+		TEST_VECTOR(ImpactPoint),
+		TEST_VECTOR(Normal),
+		TEST_VECTOR(ImpactNormal),
+		TEST_VECTOR(TraceStart),
+		TEST_VECTOR(TraceEnd),
+		TEST_FLOAT(PenetrationDepth),
+		TEST_INT(Item),
+		TEST_NAME(BoneName),
+		TEST_INT(FaceIndex)
+	};
+	
+	for(bool bSuccess : bResults)
+	{
+		if(!bSuccess)
+		{
+			return false;
+		}
+	}
+
+	return true;
+}
+
+#undef TEST_IMPL
+#undef TEST_IMPL_TOLERANCE
+#undef TEST_BOOL
+#undef TEST_INT
+#undef TEST_FLOAT
+#undef TEST_VECTOR
+
+bool FTraceQueryTestResultsInnerMost::AssertEqual(const FTraceQueryTestResultsInnerMost& Expected, const FString& What, const UObject* ContextObject, AFunctionalTest& FunctionalTest) const
+{
+	bool bSuccess = true;
+	if(FunctionalTest.AssertEqual_Bool(bSingleResult, Expected.bSingleResult, FString::Printf(TEXT("%s:bSingleResult"), *What), ContextObject))
+	{
+		bSuccess = HelperAssertFHitResultEqual(SingleHit, Expected.SingleHit, FString::Printf(TEXT("%s_SingleHit"), *What), ContextObject, FunctionalTest);
+		bSuccess &= HelperAssertNamesEqual(SingleNames, Expected.SingleNames, FString::Printf(TEXT("%s_SingleHit"), *What), ContextObject, FunctionalTest);
+	}
+	else
+	{
+		bSuccess = false;
+	}
+
+	if (FunctionalTest.AssertEqual_Bool(bMultiResult, Expected.bMultiResult, FString::Printf(TEXT("%s:bMultiResult"), *What), ContextObject))
+	{
+		if(FunctionalTest.AssertValue_Int(MultiHits.Num(), EComparisonMethod::Equal_To, Expected.MultiHits.Num(), FString::Printf(TEXT("%s:MultiHitsNum"), *What), ContextObject))
+		{
+			for (int32 HitIdx = 0; HitIdx < MultiHits.Num(); ++HitIdx)
+			{
+				bSuccess &= HelperAssertFHitResultEqual(MultiHits[HitIdx], Expected.MultiHits[HitIdx], FString::Printf(TEXT("%s_MultiHit[%d]"), *What, HitIdx), ContextObject, FunctionalTest);
+				bSuccess &= HelperAssertNamesEqual(MultiNames[HitIdx], Expected.MultiNames[HitIdx], FString::Printf(TEXT("%s_MultiHitHit[%d]"), *What, HitIdx), ContextObject, FunctionalTest);
+			}
+		}
+		else
+		{
+			bSuccess = false;
+		}
+	}
+	else
+	{
+		bSuccess = false;
+	}
+
+	return bSuccess;
+}
+
+bool FTraceQueryTestResultsInner::AssertEqual(const FTraceQueryTestResultsInner& Expected, const FString& What, const UObject* ContextObject, AFunctionalTest& FunctionalTest) const
+{
+	const bool bLineResults = LineResults.AssertEqual(Expected.LineResults, FString::Printf(TEXT("%s_Line"), *What), ContextObject, FunctionalTest);
+	const bool bSphereResults = SphereResults.AssertEqual(Expected.SphereResults, FString::Printf(TEXT("%s_Sphere"), *What), ContextObject, FunctionalTest);
+	const bool bCapsuleResults = CapsuleResults.AssertEqual(Expected.CapsuleResults, FString::Printf(TEXT("%s_Capsule"), *What), ContextObject, FunctionalTest);
+	const bool bBoxResults = BoxResults.AssertEqual(Expected.BoxResults, FString::Printf(TEXT("%s_Box"), *What), ContextObject, FunctionalTest);
+
+	return bLineResults && bSphereResults && bCapsuleResults && bBoxResults;
+}
+
+bool UTraceQueryTestResults::AssertEqual(const UTraceQueryTestResults* Expected, const FString& What, const UObject* ContextObject, AFunctionalTest& FunctionalTest) const
+{
+	if(Expected->BatchOptions != BatchOptions)
+	{
+		FunctionalTest.LogStep(ELogVerbosity::Error, FString::Printf(TEXT("Expected '%s' to be {%s} but it was {%s} for context '%s'"), *What, *Expected->BatchOptions.ToString(), *BatchOptions.ToString(), ContextObject ? *ContextObject->GetName() : TEXT("")));
+		return false;
+	}
+	else
+	{
+		//Purposely do all 3 so we get errors for as much as possible
+		const bool bChannelResults = ChannelResults.AssertEqual(Expected->ChannelResults, FString::Printf(TEXT("%s_Channel"), *What), ContextObject, FunctionalTest);
+		const bool bObjectResults = ObjectResults.AssertEqual(Expected->ObjectResults, FString::Printf(TEXT("%s_Object"), *What), ContextObject, FunctionalTest);
+		const bool bProfileResults = ProfileResults.AssertEqual(Expected->ProfileResults, FString::Printf(TEXT("%s_Profile"), *What), ContextObject, FunctionalTest);
+		if( bChannelResults && bObjectResults && bProfileResults)
+		{
+			FunctionalTest.LogStep(ELogVerbosity::Log, FString::Printf(TEXT("TraceQueryTestResults assertion passed (%s)"), *What));
+			return true;
+		}
+		else
+		{
+			FunctionalTest.LogStep(ELogVerbosity::Error, FString::Printf(TEXT("'%s' comparison failed for context '%s'"), *What, ContextObject ? *ContextObject->GetName() : TEXT("")));
+			return false;
+		}
+	}
+}
+
diff --git a/Engine/Source/Developer/FunctionalTesting/Public/AutomationBlueprintFunctionLibrary.h b/Engine/Source/Developer/FunctionalTesting/Public/AutomationBlueprintFunctionLibrary.h
--- a/Engine/Source/Developer/FunctionalTesting/Public/AutomationBlueprintFunctionLibrary.h
+++ b/Engine/Source/Developer/FunctionalTesting/Public/AutomationBlueprintFunctionLibrary.h
@@ -11,6 +11,9 @@
 
 class ACameraActor;
 
+/**
+ * 
+ */
 UCLASS()
 class FUNCTIONALTESTING_API UAutomationBlueprintFunctionLibrary : public UBlueprintFunctionLibrary
 {
@@ -19,12 +22,21 @@ class FUNCTIONALTESTING_API UAutomationBlueprintFunctionLibrary : public UBluepr
 public:
 	static bool TakeAutomationScreenshotInternal(UObject* WorldContextObject, const FString& Name, FAutomationScreenshotOptions Options);
 
+	/**
+	 * Takes a screenshot of the game's viewport.  Does not capture any UI.
+	 */
 	UFUNCTION(BlueprintCallable, Category = "Automation", meta = (Latent, HidePin = "WorldContextObject", DefaultToSelf = "WorldContextObject", LatentInfo = "LatentInfo", Name = "" ))
 	static void TakeAutomationScreenshot(UObject* WorldContextObject, FLatentActionInfo LatentInfo, const FString& Name, const FAutomationScreenshotOptions& Options);
 
+	/**
+	 * Takes a screenshot of the game's viewport, from a particular camera actors POV.  Does not capture any UI.
+	 */
 	UFUNCTION(BlueprintCallable, Category = "Automation", meta = (Latent, HidePin = "WorldContextObject", DefaultToSelf = "WorldContextObject", LatentInfo = "LatentInfo", NameOverride = "" ))
 	static void TakeAutomationScreenshotAtCamera(UObject* WorldContextObject, FLatentActionInfo LatentInfo, ACameraActor* Camera, const FString& NameOverride, const FAutomationScreenshotOptions& Options);
 
+	/**
+	 * 
+	 */
 	UFUNCTION(BlueprintCallable, Category = "Automation", meta = ( Latent, HidePin = "WorldContextObject", DefaultToSelf = "WorldContextObject", LatentInfo = "LatentInfo", NameOverride = "" ))
 	static void TakeAutomationScreenshotOfUI(UObject* WorldContextObject, FLatentActionInfo LatentInfo, const FString& Name, const FAutomationScreenshotOptions& Options);
 
@@ -49,12 +61,21 @@ public:
 	UFUNCTION(BlueprintCallable, Category = "Automation")
 	static float GetStatCallCount(FName StatName);
 
+	/**
+	 * Lets you know if any automated tests are running, or are about to run and the automation system is spinning up tests.
+	 */
 	UFUNCTION(BlueprintPure, Category="Automation")
 	static bool AreAutomatedTestsRunning();
 
+	/**
+	 * 
+	 */
 	UFUNCTION(BlueprintPure, Category="Automation")
-	static FAutomationScreenshotOptions GetDefaultScreenshotOptionsForGameplay(EComparisonTolerance Tolerance = EComparisonTolerance::Low);
+	static FAutomationScreenshotOptions GetDefaultScreenshotOptionsForGameplay(EComparisonTolerance Tolerance = EComparisonTolerance::Low, float Delay = 0.2);
 
+	/**
+	 * 
+	 */
 	UFUNCTION(BlueprintPure, Category="Automation")
-	static FAutomationScreenshotOptions GetDefaultScreenshotOptionsForRendering(EComparisonTolerance Tolerance = EComparisonTolerance::Low);
+	static FAutomationScreenshotOptions GetDefaultScreenshotOptionsForRendering(EComparisonTolerance Tolerance = EComparisonTolerance::Low, float Delay = 0.2);
 };
diff --git a/Engine/Source/Developer/FunctionalTesting/Public/GroundTruthData.h b/Engine/Source/Developer/FunctionalTesting/Public/GroundTruthData.h
new file mode 100644
--- /dev/null
+++ b/Engine/Source/Developer/FunctionalTesting/Public/GroundTruthData.h
@@ -0,0 +1,44 @@
+// Copyright 1998-2017 Epic Games, Inc. All Rights Reserved.
+
+#pragma once
+
+#include "CoreMinimal.h"
+#include "UObject/ObjectMacros.h"
+#include "UObject/Object.h"
+#include "UObject/ScriptMacros.h"
+#include "GroundTruthData.generated.h"
+
+/**
+ * 
+ */
+UCLASS(BlueprintType)
+class FUNCTIONALTESTING_API UGroundTruthData : public UObject
+{
+	GENERATED_BODY()
+
+public:
+	UGroundTruthData();
+
+	UFUNCTION(BlueprintCallable, Category = "Automation")
+	void SaveObject(UObject* GroundTruth);
+
+	UFUNCTION(BlueprintCallable, Category = "Automation")
+	UObject* LoadObject();
+
+	UFUNCTION(BlueprintCallable, Category = "Automation")
+	bool CanModify() const;
+
+#if WITH_EDITOR
+	virtual void PostEditChangeProperty(struct FPropertyChangedEvent& PropertyChangedEvent) override;
+#endif
+
+public:
+
+	UPROPERTY(EditAnywhere, Category = Data)
+	bool bResetGroundTruth;
+
+protected:
+	
+	UPROPERTY(VisibleAnywhere, Instanced, Category=Data)
+	UObject* ObjectData;
+};
diff --git a/Engine/Source/Editor/UnrealEd/Private/Commandlets/AssetRegistryGenerator.cpp b/Engine/Source/Editor/UnrealEd/Private/Commandlets/AssetRegistryGenerator.cpp
--- a/Engine/Source/Editor/UnrealEd/Private/Commandlets/AssetRegistryGenerator.cpp
+++ b/Engine/Source/Editor/UnrealEd/Private/Commandlets/AssetRegistryGenerator.cpp
@@ -428,17 +428,7 @@ bool FAssetRegistryGenerator::SaveManifests(FSandboxPlatformFile* InSandboxFile)
 
 bool FAssetRegistryGenerator::ContainsMap(const FName& PackageName) const
 {
-	const TArray<const FAssetData*>& PackageAssets = State.GetAssetsByPackageName(PackageName);
-	
-	for (const FAssetData* AssetData : PackageAssets)
-	{
-		UClass* AssetClass = AssetData->GetClass();
-		if (AssetClass->IsChildOf(UWorld::StaticClass()) || AssetClass->IsChildOf(ULevel::StaticClass()))
-		{
-			return true;
-		}
-	}
-	return false;
+	return PackagesContainingMaps.Contains(PackageName);
 }
 
 FAssetPackageData* FAssetRegistryGenerator::GetAssetPackageData(const FName& PackageName)
diff --git a/Engine/Source/Editor/UnrealEd/Private/EditorEngine.cpp b/Engine/Source/Editor/UnrealEd/Private/EditorEngine.cpp
--- a/Engine/Source/Editor/UnrealEd/Private/EditorEngine.cpp
+++ b/Engine/Source/Editor/UnrealEd/Private/EditorEngine.cpp
@@ -6991,16 +6991,18 @@ void UEditorEngine::AutomationLoadMap(const FString& MapName, FString* OutError)
 	{
 		if (Context.World())
 		{
+			FString WorldPackage = Context.World()->GetOutermost()->GetName();
+
 			if (Context.WorldType == EWorldType::PIE)
 			{
 				//don't quit!  This was triggered while pie was already running!
-				bNeedPieStart = !MapName.Contains(Context.World()->GetName());
+				bNeedPieStart = MapName != UWorld::StripPIEPrefixFromPackageName(WorldPackage, Context.World()->StreamingLevelsPrefix);
 				bPieRunning = true;
 				break;
 			}
 			else if (Context.WorldType == EWorldType::Editor)
 			{
-				bNeedLoadEditorMap = !MapName.Contains(Context.World()->GetName());
+				bNeedLoadEditorMap = MapName != WorldPackage;
 			}
 		}
 	}
diff --git a/Engine/Source/Runtime/AutomationMessages/Classes/AutomationWorkerMessages.h b/Engine/Source/Runtime/AutomationMessages/Classes/AutomationWorkerMessages.h
--- a/Engine/Source/Runtime/AutomationMessages/Classes/AutomationWorkerMessages.h
+++ b/Engine/Source/Runtime/AutomationMessages/Classes/AutomationWorkerMessages.h
@@ -643,3 +643,92 @@ public:
 	UPROPERTY(EditAnywhere, Category="Message")
 	FString ErrorMessage;
 };
+
+
+/**
+ * Implements a message that handles both storing and requesting ground truth data.
+ * for the first time this test is run, it might need to store things, or get things.
+ */
+USTRUCT()
+struct FAutomationWorkerTestDataRequest
+{
+	GENERATED_USTRUCT_BODY()
+
+	/** The category of the data, this is purely to bucket and separate the ground truth data we store into different directories. */
+	UPROPERTY(EditAnywhere, Category="Message")
+	FString DataType;
+
+	/**  */
+	UPROPERTY(EditAnywhere, Category="Message")
+	FString DataPlatform;
+
+	/**  */
+	UPROPERTY(EditAnywhere, Category="Message")
+	FString DataTestName;
+
+	/**  */
+	UPROPERTY(EditAnywhere, Category="Message")
+	FString DataName;
+
+	/**  */
+	UPROPERTY(EditAnywhere, Category="Message")
+	FString JsonData;
+};
+
+/**
+ * Implements a message that responds to TestDataRequests.
+ */
+USTRUCT()
+struct FAutomationWorkerTestDataResponse
+{
+	GENERATED_USTRUCT_BODY()
+
+	/**  */
+	UPROPERTY(EditAnywhere, Category="Message")
+	FString JsonData;
+
+	UPROPERTY(EditAnywhere, Category="Message")
+	bool bIsNew;
+};
+
+/**
+ * Implements a message to request the performance data for this hardware.
+ */
+USTRUCT()
+struct FAutomationWorkerPerformanceDataRequest
+{
+	GENERATED_USTRUCT_BODY()
+
+	/**  */
+	UPROPERTY(EditAnywhere, Category="Message")
+	FString Platform;
+
+	/**  */
+	UPROPERTY(EditAnywhere, Category="Message")
+	FString Hardware;
+
+	/**  */
+	UPROPERTY(EditAnywhere, Category="Message")
+	FString TestName;
+
+	/**  */
+	UPROPERTY(EditAnywhere, Category="Message")
+	TArray<double> DataPoints;
+};
+
+/**
+ * Implements a message that responds to PerformanceDataRequest.
+ */
+USTRUCT()
+struct FAutomationWorkerPerformanceDataResponse
+{
+	GENERATED_USTRUCT_BODY()
+
+	/**  */
+	UPROPERTY(EditAnywhere, Category="Message")
+	bool bSuccess;
+
+	/**  */
+	UPROPERTY(EditAnywhere, Category="Message")
+	FString ErrorMessage;
+};
diff --git a/Engine/Source/Runtime/AutomationWorker/Private/AutomationWorkerModule.cpp b/Engine/Source/Runtime/AutomationWorker/Private/AutomationWorkerModule.cpp
--- a/Engine/Source/Runtime/AutomationWorker/Private/AutomationWorkerModule.cpp
+++ b/Engine/Source/Runtime/AutomationWorker/Private/AutomationWorkerModule.cpp
@@ -136,6 +136,7 @@ void FAutomationWorkerModule::Initialize()
 			.Handling<FAutomationWorkerRequestTests>(this, &FAutomationWorkerModule::HandleRequestTestsMessage)
 			.Handling<FAutomationWorkerRunTests>(this, &FAutomationWorkerModule::HandleRunTestsMessage)
 			.Handling<FAutomationWorkerImageComparisonResults>(this, &FAutomationWorkerModule::HandleScreenShotCompared)
+			.Handling<FAutomationWorkerTestDataResponse>(this, &FAutomationWorkerModule::HandleTestDataRetrieved)
 			.WithInbox();
 
 		if (MessageEndpoint.IsValid())
@@ -392,6 +393,16 @@ void FAutomationWorkerModule::HandleScreenShotCompared(const FAutomationWorkerIm
 	FAutomationTestFramework::Get().NotifyScreenshotComparisonComplete(Message.bNew, Message.bSimilar, Message.MaxLocalDifference, Message.GlobalDifference, Message.ErrorMessage);
 }
 
+void FAutomationWorkerModule::HandleTestDataRetrieved(const FAutomationWorkerTestDataResponse& Message, const IMessageContextRef& Context)
+{
+	FAutomationTestFramework::Get().NotifyTestDataRetrieved(Message.bIsNew, Message.JsonData);
+}
+
+void FAutomationWorkerModule::HandlePerformanceDataRetrieved(const FAutomationWorkerPerformanceDataResponse& Message, const IMessageContextRef& Context)
+{
+	FAutomationTestFramework::Get().NotifyPerformanceDataRetrieved(Message.bSuccess, Message.ErrorMessage);
+}
+
 #if WITH_ENGINE
 void FAutomationWorkerModule::HandleScreenShotCaptured(int32 Width, int32 Height, const TArray<FColor>& Bitmap)
 {
diff --git a/Engine/Source/Runtime/AutomationWorker/Private/AutomationWorkerModule.h b/Engine/Source/Runtime/AutomationWorker/Private/AutomationWorkerModule.h
--- a/Engine/Source/Runtime/AutomationWorker/Private/AutomationWorkerModule.h
+++ b/Engine/Source/Runtime/AutomationWorker/Private/AutomationWorkerModule.h
@@ -10,6 +10,8 @@
 
 struct FAutomationWorkerFindWorkers;
 struct FAutomationWorkerImageComparisonResults;
+struct FAutomationWorkerTestDataResponse;
+struct FAutomationWorkerPerformanceDataResponse;
 struct FAutomationWorkerNextNetworkCommandReply;
 struct FAutomationWorkerPing;
 struct FAutomationWorkerRequestTests;
@@ -104,6 +106,12 @@ private:
 	// Handles FAutomationWorkerImageComparisonResults messages.
 	void HandleScreenShotCompared(const FAutomationWorkerImageComparisonResults& Message, const IMessageContextRef& Context);
 
+	// Handles FAutomationWorkerTestDataResponse messages.
+	void HandleTestDataRetrieved(const FAutomationWorkerTestDataResponse& Message, const IMessageContextRef& Context);
+
+	// Handles FAutomationWorkerPerformanceDataResponse messages.
+	void HandlePerformanceDataRetrieved(const FAutomationWorkerPerformanceDataResponse& Message, const IMessageContextRef& Context);
+
 	// Handles FAutomationTestFramework PreTestingEvents.
 	void HandlePreTestingEvent();
 
diff --git a/Engine/Source/Runtime/Core/Private/Misc/AutomationTest.cpp b/Engine/Source/Runtime/Core/Private/Misc/AutomationTest.cpp
--- a/Engine/Source/Runtime/Core/Private/Misc/AutomationTest.cpp
+++ b/Engine/Source/Runtime/Core/Private/Misc/AutomationTest.cpp
@@ -11,6 +11,7 @@
 #include "Misc/App.h"
 #include "Modules/ModuleManager.h"
 #include "Misc/OutputDeviceRedirector.h"
+#include "Internationalization/Regex.h"
 
 
 DEFINE_LOG_CATEGORY_STATIC(LogAutomationTest, Warning, All);
@@ -606,10 +607,11 @@ bool FAutomationTestFramework::InternalStopTest(FAutomationTestExecutionInfo& Ou
 	// Disassociate the test from the feedback context
 	AutomationTestFeedbackContext.SetCurrentAutomationTest( NULL );
 
-	// Determine if the test was successful based on two criteria:
+	// Determine if the test was successful based on three criteria:
 	// 1) Did the test itself report success?
 	// 2) Did any errors occur and were logged by the feedback context during execution?++----
-	bTestSuccessful = bTestSuccessful && !CurrentTest->HasAnyErrors();
+	// 3) Did we meet any errors that were expected with this test
+	bTestSuccessful = bTestSuccessful && !CurrentTest->HasAnyErrors() && CurrentTest->HasMetExpectedErrors();
 
 	// Set the success state of the test based on the above criteria
 	CurrentTest->SetSuccessState( bTestSuccessful );
@@ -617,10 +619,17 @@ bool FAutomationTestFramework::InternalStopTest(FAutomationTestExecutionInfo& Ou
 	// Fill out the provided execution info with the info from the test
 	CurrentTest->GetExecutionInfo( OutExecutionInfo );
 
-	//save off timing for the test
+	// Save off timing for the test
 	OutExecutionInfo.Duration = TimeForTest;
 
-	//release pointers to now-invalid data
+	// Re-enable log parsing if it was disabled and empty the expected errors list
+	if (CurrentTest->ExpectedErrors.Num())
+	{
+		GLog->Logf(ELogVerbosity::Display, TEXT("<-- Resume Log Parsing -->"));
+	}
+	CurrentTest->ExpectedErrors.Empty();
+
+	// Release pointers to now-invalid data
 	CurrentTest = NULL;
 
 	return bTestSuccessful;
@@ -653,6 +662,16 @@ void FAutomationTestFramework::NotifyScreenshotComparisonComplete(bool bWasNew,
 	OnScreenshotCompared.Broadcast(bWasNew, bWasSimilar, MaxLocalDifference, GlobalDifference, ErrorMessage);
 }
 
+void FAutomationTestFramework::NotifyTestDataRetrieved(bool bWasNew, const FString& JsonData)
+{
+	OnTestDataRetrieved.Broadcast(bWasNew, JsonData);
+}
+
+void FAutomationTestFramework::NotifyPerformanceDataRetrieved(bool bSuccess, const FString& ErrorMessage)
+{
+	OnPerformanceDataRetrieved.Broadcast(bSuccess, ErrorMessage);
+}
+
 void FAutomationTestFramework::NotifyScreenshotTakenAndCompared()
 {
 	OnScreenshotTakenAndCompared.Broadcast();
@@ -782,7 +801,7 @@ void FAutomationTestBase::ClearExecutionInfo()
 
 void FAutomationTestBase::AddError(const FString& InError, int32 StackOffset)
 {
-	if( !bSuppressLogs )
+	if( !bSuppressLogs && !IsExpectedError(InError))
 	{
 		if ( FAutomationTestFramework::Get().GetCaptureStack() )
 		{
@@ -798,7 +817,7 @@ void FAutomationTestBase::AddError(const FString& InError, int32 StackOffset)
 
 void FAutomationTestBase::AddError(const FString& InError, const FString& InFilename, int32 InLineNumber)
 {
-	if ( !bSuppressLogs )
+	if ( !bSuppressLogs && !IsExpectedError(InError))
 	{
 		ExecutionInfo.AddEvent(FAutomationEvent(EAutomationEventType::Error, InError, ExecutionInfo.GetContext(), InFilename, InLineNumber));
 	}
@@ -806,7 +825,7 @@ void FAutomationTestBase::AddError(const FString& InError, const FString& InFile
 
 void FAutomationTestBase::AddWarning(const FString& InWarning, const FString& InFilename, int32 InLineNumber)
 {
-	if ( !bSuppressLogs )
+	if ( !bSuppressLogs && !IsExpectedError(InWarning))
 	{
 		ExecutionInfo.AddEvent(FAutomationEvent(EAutomationEventType::Warning, InWarning, ExecutionInfo.GetContext(), InFilename, InLineNumber));
 	}
@@ -814,7 +833,7 @@ void FAutomationTestBase::AddWarning(const FString& InWarning, const FString& In
 
 void FAutomationTestBase::AddWarning( const FString& InWarning, int32 StackOffset )
 {
-	if ( !bSuppressLogs )
+	if ( !bSuppressLogs && !IsExpectedError(InWarning))
 	{
 		if ( FAutomationTestFramework::Get().GetCaptureStack() )
 		{
@@ -854,6 +873,50 @@ bool FAutomationTestBase::HasAnyErrors() const
 	return ExecutionInfo.GetErrorTotal() > 0;
 }
 
+bool FAutomationTestBase::HasMetExpectedErrors()
+{
+	bool HasMetAllExpectedErrors = true;
+
+	for (auto& EError : ExpectedErrors)
+	{
+		if ((EError.ExpectedNumberOfOccurrences > 0) && (EError.ExpectedNumberOfOccurrences != EError.ActualNumberOfOccurrences))
+		{
+			HasMetAllExpectedErrors = false;
+
+			ExecutionInfo.AddEvent(FAutomationEvent(EAutomationEventType::Error,
+				FString::Printf(TEXT("Expected Error or Warning matching '%s' to occur %d times with %s match type, but it was found %d time(s).")
+					, *EError.ErrorPatternString
+					, EError.ExpectedNumberOfOccurrences
+					, EAutomationExpectedErrorFlags::ToString(EError.CompareType)
+					, EError.ActualNumberOfOccurrences)
+				, ExecutionInfo.GetContext()));			
+		}
+		else if (EError.ExpectedNumberOfOccurrences == 0)
+		{
+			if (EError.ActualNumberOfOccurrences == 0)
+			{
+				HasMetAllExpectedErrors = false;
+
+				ExecutionInfo.AddEvent(FAutomationEvent(EAutomationEventType::Error,
+					FString::Printf(TEXT("Expected suppressed Error or Warning matching '%s' did not occur."), *EError.ErrorPatternString),
+					ExecutionInfo.GetContext()));
+			}
+			else
+			{
+				ExecutionInfo.AddEvent(FAutomationEvent(EAutomationEventType::Info,
+					FString::Printf(TEXT("Suppressed expected Error or Warning matching '%s' %d times.")
+						, *EError.ErrorPatternString
+						, EError.ActualNumberOfOccurrences)
+					, ExecutionInfo.GetContext()));
+			}
+		}
+	}
+
+	return HasMetAllExpectedErrors;
+}
+
+
+
 void FAutomationTestBase::SetSuccessState( bool bSuccessful )
 {
 	ExecutionInfo.bSuccessful = bSuccessful;
@@ -864,6 +927,45 @@ void FAutomationTestBase::GetExecutionInfo( FAutomationTestExecutionInfo& OutInf
 	OutInfo = ExecutionInfo;
 }
 
+void FAutomationTestBase::AddExpectedError(FString ExpectedErrorPattern, EAutomationExpectedErrorFlags::MatchType InCompareType, int32 Occurrences)
+{
+	if (Occurrences >= 0)
+	{
+		// If we already have an error matching string in our list, let's not add it again.
+		FAutomationExpectedError* FoundEntry = ExpectedErrors.FindByPredicate(
+			[ExpectedErrorPattern](const FAutomationExpectedError& InItem) 
+				{
+					return InItem.ErrorPatternString == ExpectedErrorPattern; 
+				}
+		);
+
+		if (FoundEntry)
+		{
+			UE_LOG(LogAutomationTest, Warning, TEXT("Adding expected error matching '%s' failed: cannot add duplicate entries"), *ExpectedErrorPattern)
+		}
+		else
+		{
+			// Disable log pre-processor the first time we successfully add an expected error
+			// so that successful tests don't trigger CIS failures
+			if (!ExpectedErrors.Num())
+			{
+				GLog->Logf(ELogVerbosity::Display, TEXT("<-- Suspend Log Parsing -->"));
+			}
+			// ToDo: After UE-44340 is resolved, create FAutomationExpectedError and check that its ErrorPattern is valid before adding
+			ExpectedErrors.Add(FAutomationExpectedError(ExpectedErrorPattern, InCompareType, Occurrences));
+		}
+	}
+	else
+	{
+		UE_LOG(LogAutomationTest, Error, TEXT("Adding expected error matching '%s' failed: number of expected occurrences must be >= 0"), *ExpectedErrorPattern);
+	}
+}
+
+void FAutomationTestBase::GetExpectedErrors(TArray<FAutomationExpectedError>& OutInfo) const
+{
+	OutInfo = ExpectedErrors;
+}
+
 void FAutomationTestBase::GenerateTestNames(TArray<FAutomationTestInfo>& TestInfo) const
 {
 	TArray<FString> BeautifiedNames;
@@ -986,3 +1088,19 @@ void FAutomationTestBase::TestNull(const FString& What, void* Pointer)
 		AddInfo(FString::Printf(TEXT("Expected '%s' to be null."), *What), 1);
 	}
 }
+
+bool FAutomationTestBase::IsExpectedError(const FString& Error)
+{
+	for (auto& EError : ExpectedErrors)
+	{
+		FRegexMatcher ErrorMatcher(EError.ErrorPattern, Error);
+
+		if (ErrorMatcher.FindNext())
+		{
+			EError.ActualNumberOfOccurrences++;
+			return true;
+		}
+	}
+
+	return false;
+}
\ No newline at end of file
diff --git a/Engine/Source/Runtime/Core/Public/Misc/AutomationTest.h b/Engine/Source/Runtime/Core/Public/Misc/AutomationTest.h
--- a/Engine/Source/Runtime/Core/Public/Misc/AutomationTest.h
+++ b/Engine/Source/Runtime/Core/Public/Misc/AutomationTest.h
@@ -26,13 +26,14 @@
 #include "Math/Color.h"
 #include "PlatformProcess.h"
 #include "Misc/AutomationEvent.h"
+#include "Internationalization/Regex.h"
 
 /** Flags for specifying automation test requirements/behavior */
 namespace EAutomationTestFlags
 {
 	enum Type
 	{
-		//~ Application context required for the test - not specifying means it will be valid for any context
+		//~ Application context required for the test
 		// Test is suitable for running within the editor
 		EditorContext = 0x00000001,
 		// Test is suitable for running within the client
@@ -84,6 +85,30 @@ namespace EAutomationTestFlags
 	};
 };
 
+/** Flags for indicating the matching type to use for an expected error */
+namespace EAutomationExpectedErrorFlags
+{
+	enum MatchType
+	{
+		// When matching expected errors, do so exactly.
+		Exact,
+		// When matching expected errors, just see if the error string is contained in the string to be evaluated.
+		Contains,
+	};
+
+	inline const TCHAR* ToString(EAutomationExpectedErrorFlags::MatchType ThisType)
+	{
+		switch (ThisType)
+		{
+		case Contains:
+			return TEXT("Contains");
+		case Exact:
+			return TEXT("Exact");
+		}
+		return TEXT("Unknown");
+	}
+}
+
 /** Simple class to store the results of the execution of a automation test */
 class CORE_API FAutomationTestExecutionInfo
 {
@@ -481,6 +506,43 @@ public:
 	virtual void Run() = 0;
 };
 
+struct FAutomationExpectedError
+{
+	// Original regular expression pattern string matching expected error message.
+	// NOTE: using the Exact comparison type wraps the pattern string with ^ and $ tokens,
+	// but the base pattern string is preserved to allow checks for duplicate entries.
+	FString ErrorPatternString;
+	// Regular expression pattern for ErrorPatternString
+	FRegexPattern ErrorPattern;
+	// Type of comparison to perform on error message using ErrorPattern.
+	EAutomationExpectedErrorFlags::MatchType CompareType;
+	/** 
+	 * Number of occurrences expected for error. If set greater than 0, it will cause the test to fail if the
+	 * exact number of occurrences expected is not matched. If set to 0, it will suppress all matching messages. 
+	 */
+	int32 ExpectedNumberOfOccurrences;
+	int32 ActualNumberOfOccurrences;
+
+	/**
+	* Constructor
+	*/
+	
+	FAutomationExpectedError(FString& InErrorPattern, EAutomationExpectedErrorFlags::MatchType InCompareType, int32 InExpectedNumberOfOccurrences = 1)
+		: ErrorPatternString(InErrorPattern)
+		, ErrorPattern((InCompareType == EAutomationExpectedErrorFlags::Exact) ? FString::Printf(TEXT("^%s$"), *InErrorPattern) : InErrorPattern)
+		, CompareType(InCompareType)
+		, ExpectedNumberOfOccurrences(InExpectedNumberOfOccurrences)
+		, ActualNumberOfOccurrences(0)
+	{}
+
+	FAutomationExpectedError(FString& InErrorPattern, int32 InExpectedNumberOfOccurrences)
+		: ErrorPatternString(InErrorPattern)
+		, ErrorPattern(InErrorPattern)
+		, CompareType(EAutomationExpectedErrorFlags::Contains)
+		, ExpectedNumberOfOccurrences(InExpectedNumberOfOccurrences)
+	{}
+};
+
 struct FAutomationScreenshotData
 {
 	FString Name;
@@ -571,6 +633,10 @@ DECLARE_DELEGATE_TwoParams(FOnTestScreenshotCaptured, const TArray<FColor>&, con
 
 DECLARE_MULTICAST_DELEGATE_FiveParams(FOnTestScreenshotComparisonComplete, bool /*bWasNew*/, bool /*bWasSimilar*/, double /*MaxLocalDifference*/, double /*GlobalDifference*/, FString /*ErrorMessage*/);
 
+DECLARE_MULTICAST_DELEGATE_TwoParams(FOnTestDataRetrieved, bool /*bWasNew*/, const FString& /*JsonData*/);
+
+DECLARE_MULTICAST_DELEGATE_TwoParams(FOnPerformanceDataRetrieved, bool /*bSuccess*/, const FString& /*ErrorMessage*/);
+
 /** Class representing the main framework for running automation tests */
 class CORE_API FAutomationTestFramework
 {
@@ -584,6 +650,12 @@ public:
 	/** Called when a screenshot comparison completes. */
 	FOnTestScreenshotComparisonComplete OnScreenshotCompared;
 
+	/** Called when the test data is retrieved. */
+	FOnTestDataRetrieved OnTestDataRetrieved;
+
+	/** Called when the performance data is retrieved. */
+	FOnPerformanceDataRetrieved OnPerformanceDataRetrieved;
+
 	/** The final call related to screenshots, after they've been taken, and after they've been compared (or not if automation isn't running). */
 	FSimpleMulticastDelegate OnScreenshotTakenAndCompared;
 
@@ -759,6 +831,8 @@ public:
 	void SetTreatWarningsAsErrors(TOptional<bool> bTreatWarningsAsErrors);
 
 	void NotifyScreenshotComparisonComplete(bool bWasNew, bool bWasSimilar, double MaxLocalDifference, double GlobalDifference, FString ErrorMessage);
+	void NotifyTestDataRetrieved(bool bWasNew, const FString& JsonData);
+	void NotifyPerformanceDataRetrieved(bool bSuccess, const FString& ErrorMessage);
 
 	void NotifyScreenshotTakenAndCompared();
 
@@ -990,6 +1064,13 @@ public:
 	 */
 	bool HasAnyErrors() const;
 
+	/**
+	* Returns whether this test has encountered all expected errors defined for it
+	*
+	* @return true if this test has encountered all expected errors; false if not
+	*/
+	bool HasMetExpectedErrors();
+
 	/**
 	 * Forcibly sets whether the test has succeeded or not
 	 *
@@ -1010,6 +1091,25 @@ public:
 	 */
 	void GenerateTestNames( TArray<FAutomationTestInfo>& TestInfo ) const;
 
+	/**
+	* Adds a regex pattern to an internal list that this test will expect to encounter in error or warning logs during its execution. If an expected pattern
+	* is not encountered, it will cause this test to fail.
+	*
+	* @param ExpectedPatternString - The expected message string. Supports basic regex patterns.
+	* @param CompareType - How to match this string with an encountered error, should it match exactly or simply just contain the string.
+	* @param Occurrences - How many times to expect this error string to be seen. If > 0, the error must be seen the exact number of times
+	* specified or the test will fail. If == 0, the error must be seen one or more times (with no upper limit) or the test will fail.
+	*/
+	void AddExpectedError(FString ExpectedPatternString, EAutomationExpectedErrorFlags::MatchType CompareType = EAutomationExpectedErrorFlags::Contains, int32 Occurrences = 1);
+
+	/**
+	* Populate the provided expected errors object with the expected errors contained within the test. Not particularly efficient,
+	* but providing direct access to the test's private execution errors list could result in errors.
+	*
+	* @param OutInfo - Array of Expected Errors to be populated with the same data contained within this test's expected errors list
+	*/
+	void GetExpectedErrors(TArray<FAutomationExpectedError>& OutInfo) const;
+
 	/**
 	 * Is this a complex tast - if so it will be a stress test.
 	 *
@@ -1307,6 +1407,19 @@ protected:
 
 	//allow framework to call protected function
 	friend class FAutomationTestFramework;
+
+private:
+	/**
+	* Returns whether this test has defined any expected errors matching the given message.
+	* If a match is found, the expected error definition increments it actual occurrence count.
+	*
+	* @return true if this message matches any of the expected errors
+	*/
+	bool IsExpectedError(const FString& Error);
+
+	/* Errors to be expected while processing this test.*/
+	TArray< FAutomationExpectedError> ExpectedErrors;
+
 };
 
 class CORE_API FBDDAutomationTestBase : public FAutomationTestBase
diff --git a/Engine/Source/Runtime/CoreUObject/Public/UObject/UnrealType.h b/Engine/Source/Runtime/CoreUObject/Public/UObject/UnrealType.h
--- a/Engine/Source/Runtime/CoreUObject/Public/UObject/UnrealType.h
+++ b/Engine/Source/Runtime/CoreUObject/Public/UObject/UnrealType.h
@@ -4309,18 +4309,6 @@ struct FPropertyChangedEvent
 	{
 	}
 
-
-	DEPRECATED(4.7, "The bInChangesTopology parameter has been removed, use the two-argument constructor of FPropertyChangedEvent instead")
-	FPropertyChangedEvent(UProperty* InProperty, const bool /*bInChangesTopology*/, EPropertyChangeType::Type InChangeType)
-		: Property(InProperty)
-		, MemberProperty(InProperty)
-		, ChangeType(InChangeType)
-		, ObjectIteratorIndex(INDEX_NONE)
-		, ArrayIndicesPerObject(nullptr)
-		, TopLevelObjects(nullptr)
-	{
-	}
-
 	void SetActiveMemberProperty( UProperty* InActiveMemberProperty )
 	{
 		MemberProperty = InActiveMemberProperty;
@@ -4366,6 +4354,14 @@ struct FPropertyChangedEvent
 	 */
 	const UObject* GetObjectBeingEdited(int32 Index) const { return TopLevelObjects ? (*TopLevelObjects)[Index] : nullptr; }
 
+	/**
+	 * Simple utility to get the name of the property and takes care of the possible null property.
+	 */
+	FName GetPropertyName() const
+	{
+		return (Property != nullptr) ? Property->GetFName() : NAME_None;
+	}
+
 	/**
 	 * The actual property that changed
 	 */
