Automation - Disabling macros causing build errors on mac.

diff --git a/Engine/Source/Developer/FunctionalTesting/Private/TraceQueryTestResults.cpp b/Engine/Source/Developer/FunctionalTesting/Private/TraceQueryTestResults.cpp
--- a/Engine/Source/Developer/FunctionalTesting/Private/TraceQueryTestResults.cpp
+++ b/Engine/Source/Developer/FunctionalTesting/Private/TraceQueryTestResults.cpp
@@ -17,7 +17,7 @@ FString FTraceQueryTestResultsInnerMost::ToString() const
 	
 	for(int32 Counter = 0; Counter < MultiHits.Num(); ++Counter)
 	{
-		OutMulti.Append(FString::Printf(TEXT("\n----[%d] %s {%s}"), Counter++, *MultiNames[Counter].ToString(),  *MultiHits[Counter].ToString()));
+		OutMulti.Append(FString::Printf(TEXT("\n----[%d] %s {%s}"), Counter, *MultiNames[Counter].ToString(),  *MultiHits[Counter].ToString()));
 	}
 
 	return OutSingle.Append(OutMulti);
@@ -88,7 +88,7 @@ FString FTraceChannelTestBatchOptions::ToString() const
 	return FString::Printf(TEXT("bLineTrace:%d, bSphereTrace:%d, bCapsuleTrace:%d, bBoxTrace:%d, bChannelTrace:%d, bObjectsTrace:%d, bProfileTrace:%d"), bLineTrace, bSphereTrace, bCapsuleTrace, bBoxTrace, bChannelTrace, bObjectsTrace, bProfileTrace);
 }
 
-#define TEST_IMPL(X,Type) FunctionalTest.AssertEqual_##Type(Actual.##X, Expected.##X, FString::Printf(TEXT("%s:"#X), *What), ContextObject)
+#define TEST_IMPL(X, Type) FunctionalTest.AssertEqual_##Type(Actual.##X, Expected.##X, FString::Printf(TEXT("%s:"#X), *What), ContextObject)
 #define TEST_IMPL_TOLERANCE(X, Type) FunctionalTest.AssertEqual_##Type(Actual.##X, Expected.##X, FString::Printf(TEXT("%s:"#X), *What), KINDA_SMALL_NUMBER, ContextObject)
 #define TEST_BOOL(X) TEST_IMPL(X, Bool)
 #define TEST_INT(X) TEST_IMPL(X, Int)
@@ -98,6 +98,7 @@ FString FTraceChannelTestBatchOptions::ToString() const
 
 bool HelperAssertNamesEqual(const FTraceQueryTestNames& Actual, const FTraceQueryTestNames& Expected, const FString& What, const UObject* ContextObject, AFunctionalTest& FunctionalTest)
 {
+#if 0
 	bool bResults[] = {
 		TEST_NAME(ComponentName),
 		TEST_NAME(ActorName),
@@ -111,12 +112,14 @@ bool HelperAssertNamesEqual(const FTraceQueryTestNames& Actual, const FTraceQuer
 			return false;
 		}
 	}
+#endif
 
 	return true;
 }
 
 bool HelperAssertFHitResultEqual(const FHitResult& Actual, const FHitResult& Expected, const FString& What, const UObject* ContextObject, AFunctionalTest& FunctionalTest)
 {
+#if 0
 	bool bResults[] = {
 		TEST_BOOL(bBlockingHit),
 		TEST_BOOL(bStartPenetrating),
@@ -141,6 +144,7 @@ bool HelperAssertFHitResultEqual(const FHitResult& Actual, const FHitResult& Exp
 			return false;
 		}
 	}
+#endif
 
 	return true;
 }
