Build - Fixing a PS4 build issue with ground truth data.

diff --git a/Engine/Source/Developer/FunctionalTesting/Private/GroundTruthData.cpp b/Engine/Source/Developer/FunctionalTesting/Private/GroundTruthData.cpp
--- a/Engine/Source/Developer/FunctionalTesting/Private/GroundTruthData.cpp
+++ b/Engine/Source/Developer/FunctionalTesting/Private/GroundTruthData.cpp
@@ -2,14 +2,15 @@
 
 #include "GroundTruthData.h"
 
+#include "UObject/UnrealType.h"
+#include "UObject/Package.h"
+#include "Misc/PackageName.h"
+#include "AssetData.h"
+
 #if WITH_EDITOR
 #include "ISourceControlModule.h"
 #include "ISourceControlProvider.h"
 #include "ISourceControlOperation.h"
-#include "UObject/Package.h"
-#include "Misc/PackageName.h"
-#include "AssetData.h"
-#include "UObject/UnrealType.h"
 #endif
 
 DEFINE_LOG_CATEGORY_STATIC(GroundTruthLog, Log, Log)
