Fix automation warnings

diff --git a/Engine/Source/Developer/FunctionalTesting/Classes/ScreenshotFunctionalTest.h b/Engine/Source/Developer/FunctionalTesting/Classes/ScreenshotFunctionalTest.h
--- a/Engine/Source/Developer/FunctionalTesting/Classes/ScreenshotFunctionalTest.h
+++ b/Engine/Source/Developer/FunctionalTesting/Classes/ScreenshotFunctionalTest.h
@@ -37,8 +37,8 @@ protected:
 	/** The horizontal field of view (in degrees) for the screenshot camera */
 	UPROPERTY(EditANywhere, Category="Screenshot", meta=(UIMin = "5.0", UIMax = "170", ClampMin = "0.001", ClampMax = "360.0", Units = deg))
 	float FieldOfView;
-	
-	UPROPERTY()
+
+	UPROPERTY(BlueprintReadOnly, Category = "Camera", SimpleDisplay)
 	class UCameraComponent* ScreenshotCamera;
 
 	UPROPERTY(EditAnywhere, BlueprintReadOnly, Category="Screenshot", SimpleDisplay)
