Copy CL 3514743 from Dev-Automation

Back out changelist 3491768

Removed FOV option from ScreenshotFunctionalTest.
Exposed ScreenshotFTest camera component as editable instead.

#tests UFE and preflight

diff --git a/Engine/Source/Developer/FunctionalTesting/Classes/ScreenshotFunctionalTest.h b/Engine/Source/Developer/FunctionalTesting/Classes/ScreenshotFunctionalTest.h
--- a/Engine/Source/Developer/FunctionalTesting/Classes/ScreenshotFunctionalTest.h
+++ b/Engine/Source/Developer/FunctionalTesting/Classes/ScreenshotFunctionalTest.h
@@ -33,12 +33,8 @@ protected:
 	void OnScreenshotTakenAndCompared();
 	
 protected:
-
-	/** The horizontal field of view (in degrees) for the screenshot camera */
-	UPROPERTY(EditANywhere, Category="Screenshot", meta=(UIMin = "5.0", UIMax = "170", ClampMin = "0.001", ClampMax = "360.0", Units = deg))
-	float FieldOfView;
-
-	UPROPERTY(BlueprintReadOnly, Category = "Camera", SimpleDisplay)
+	
+	UPROPERTY(EditAnywhere, BlueprintReadWrite, Category = "Camera", SimpleDisplay)
 	class UCameraComponent* ScreenshotCamera;
 
 	UPROPERTY(EditAnywhere, BlueprintReadOnly, Category="Screenshot", SimpleDisplay)
diff --git a/Engine/Source/Developer/FunctionalTesting/Private/ScreenshotFunctionalTest.cpp b/Engine/Source/Developer/FunctionalTesting/Private/ScreenshotFunctionalTest.cpp
--- a/Engine/Source/Developer/FunctionalTesting/Private/ScreenshotFunctionalTest.cpp
+++ b/Engine/Source/Developer/FunctionalTesting/Private/ScreenshotFunctionalTest.cpp
@@ -17,7 +17,6 @@ AScreenshotFunctionalTest::AScreenshotFunctionalTest( const FObjectInitializer&
 {
 	ScreenshotCamera = CreateDefaultSubobject<UCameraComponent>(TEXT("Camera"));
 	ScreenshotCamera->SetupAttachment(RootComponent);
-	FieldOfView = ScreenshotCamera->FieldOfView;
 }
 
 void AScreenshotFunctionalTest::PrepareTest()
@@ -92,10 +91,6 @@ void AScreenshotFunctionalTest::PostEditChangeProperty(FPropertyChangedEvent& Pr
 		{
 			ScreenshotOptions.SetToleranceAmounts(ScreenshotOptions.Tolerance);
 		}
-		else if (PropertyName == GET_MEMBER_NAME_CHECKED(AScreenshotFunctionalTest, FieldOfView))
-		{
-			ScreenshotCamera->SetFieldOfView(FieldOfView);
-		}
 	}
 }
 
