Key changes:
- Disable texture streaming
- AScreenshotFunctionalTest now reads back from game viewport directly when taking screenshot
- AFunctionalUIScreenshotTest now render into a custom render target using FWidgetRenderer so that UI screenshot can be captured at fixed size
- Add a new class, AScreenshotFunctionalTestBase, that handles screenshot common functions. AScreenshotFunctionalTest and AFunctionalUIScreenshotTest now derive from it

diff --git a/Engine/Source/Developer/FunctionalTesting/Classes/FunctionalUIScreenshotTest.h b/Engine/Source/Developer/FunctionalTesting/Classes/FunctionalUIScreenshotTest.h
--- a/Engine/Source/Developer/FunctionalTesting/Classes/FunctionalUIScreenshotTest.h
+++ b/Engine/Source/Developer/FunctionalTesting/Classes/FunctionalUIScreenshotTest.h
@@ -4,12 +4,14 @@
 
 #include "CoreMinimal.h"
 #include "UObject/ObjectMacros.h"
-#include "FunctionalTest.h"
+#include "ScreenshotFunctionalTestBase.h"
 #include "AutomationScreenshotOptions.h"
 #include "Blueprint/UserWidget.h"
 
 #include "FunctionalUIScreenshotTest.generated.h"
 
+class UTextureRenderTarget2D;
+
 UENUM()
 enum class EWidgetTestAppearLocation
 {
@@ -21,24 +23,19 @@ enum class EWidgetTestAppearLocation
  * 
  */
 UCLASS(Blueprintable, MinimalAPI)
-class AFunctionalUIScreenshotTest : public AFunctionalTest
+class AFunctionalUIScreenshotTest : public AScreenshotFunctionalTestBase
 {
 	GENERATED_BODY()
 
 public:
 	AFunctionalUIScreenshotTest(const FObjectInitializer& ObjectInitializer);
 
-#if WITH_EDITOR
-	virtual bool CanEditChange(const UProperty* InProperty) const override;
-	virtual void PostEditChangeProperty(FPropertyChangedEvent& PropertyChangedEvent) override;
-#endif
-
 protected:
 	virtual void PrepareTest() override;
-	virtual bool IsReady_Implementation() override;
-	virtual void StartTest() override;
 
-	void OnScreenshotTakenAndCompared();
+	virtual void OnScreenshotTakenAndCompared() override;
+
+	virtual void RequestScreenshot() override;
 	
 protected:
 
@@ -50,10 +47,7 @@ protected:
 
 	UPROPERTY(EditAnywhere, Category = "UI")
 	EWidgetTestAppearLocation WidgetLocation;
-	
-	UPROPERTY()
-	class UCameraComponent* ScreenshotCamera;
 
-	UPROPERTY(EditAnywhere, Category="Screenshot", SimpleDisplay)
-	FAutomationScreenshotOptions ScreenshotOptions;
+	UPROPERTY(Transient, DuplicateTransient)
+	UTextureRenderTarget2D* ScreenshotRT;
 };
diff --git a/Engine/Source/Developer/FunctionalTesting/Classes/ScreenshotFunctionalTest.h b/Engine/Source/Developer/FunctionalTesting/Classes/ScreenshotFunctionalTest.h
--- a/Engine/Source/Developer/FunctionalTesting/Classes/ScreenshotFunctionalTest.h
+++ b/Engine/Source/Developer/FunctionalTesting/Classes/ScreenshotFunctionalTest.h
@@ -4,43 +4,30 @@
 
 #include "CoreMinimal.h"
 #include "UObject/ObjectMacros.h"
-#include "FunctionalTest.h"
+#include "ScreenshotFunctionalTestBase.h"
 #include "AutomationScreenshotOptions.h"
 
 #include "ScreenshotFunctionalTest.generated.h"
 
+class FAutomationTestScreenshotEnvSetup;
+
 /**
- * 
+ * No UI
  */
 UCLASS(Blueprintable)
-class FUNCTIONALTESTING_API AScreenshotFunctionalTest : public AFunctionalTest
+class FUNCTIONALTESTING_API AScreenshotFunctionalTest : public AScreenshotFunctionalTestBase
 {
 	GENERATED_BODY()
 
 public:
 	AScreenshotFunctionalTest(const FObjectInitializer& ObjectInitializer);
 
-#if WITH_EDITOR
-	virtual bool CanEditChange(const UProperty* InProperty) const override;
-	virtual void PostEditChangeProperty(FPropertyChangedEvent& PropertyChangedEvent) override;
-#endif
-
 	// Tests relying on temporal effects can force a camera cut to flush stale data
 	UPROPERTY(EditAnywhere, BlueprintReadWrite, Category = "Camera", SimpleDisplay)
 	bool bCameraCutOnScreenshotPrep;
 
 protected:
 	virtual void PrepareTest() override;
-	virtual bool IsReady_Implementation() override;
-	virtual void StartTest() override;
-
-	void OnScreenshotTakenAndCompared();
-	
-protected:
-	
-	UPROPERTY(EditAnywhere, BlueprintReadWrite, Category = "Camera", SimpleDisplay)
-	class UCameraComponent* ScreenshotCamera;
 
-	UPROPERTY(EditAnywhere, BlueprintReadOnly, Category="Screenshot", SimpleDisplay)
-	FAutomationScreenshotOptions ScreenshotOptions;
+	virtual void RequestScreenshot() override;
 };
diff --git a/Engine/Source/Developer/FunctionalTesting/Classes/ScreenshotFunctionalTestBase.h b/Engine/Source/Developer/FunctionalTesting/Classes/ScreenshotFunctionalTestBase.h
new file mode 100644
--- /dev/null
+++ b/Engine/Source/Developer/FunctionalTesting/Classes/ScreenshotFunctionalTestBase.h
@@ -0,0 +1,73 @@
+// Copyright 1998-2017 Epic Games, Inc. All Rights Reserved.
+
+#pragma once
+
+#include "CoreMinimal.h"
+#include "UObject/ObjectMacros.h"
+#include "FunctionalTest.h"
+#include "AutomationScreenshotOptions.h"
+
+#include "ScreenshotFunctionalTestBase.generated.h"
+
+class FAutomationTestScreenshotEnvSetup;
+
+/**
+* Base class for screenshot functional test
+*/
+UCLASS(Blueprintable, abstract)
+class FUNCTIONALTESTING_API AScreenshotFunctionalTestBase : public AFunctionalTest
+{
+	GENERATED_BODY()
+
+public:
+	AScreenshotFunctionalTestBase(const FObjectInitializer& ObjectInitializer);
+
+#if WITH_EDITOR
+	virtual bool CanEditChange(const UProperty* InProperty) const override;
+	virtual void PostEditChangeProperty(FPropertyChangedEvent& PropertyChangedEvent) override;
+#endif
+
+protected:
+	// Set player view target to screenshot camera and call PrepareForScreenshot
+	virtual void PrepareTest() override;
+
+	// Handle screenshot delay
+	virtual bool IsReady_Implementation() override;
+
+	// Register OnScreenshotTakenAndCompared and call RequestScreenshot
+	virtual void StartTest() override;
+
+	// Call RestoreViewport and finish this test
+	virtual void OnScreenshotTakenAndCompared();
+
+	// Resize viewport to screenshot size (if possible) and set up screenshot environment (disable AA, etc.)
+	void PrepareForScreenshot();
+
+	// Doesn't actually request in base class. It simply register OnScreenshotCaptured
+	virtual void RequestScreenshot();
+
+	// Pass screenshot pixels and meta data to FAutomationTestFramework. Register
+	// OnComparisonComplete which will be called the automation test system when
+	// screenshot comparison is complete
+	void OnScreenShotCaptured(int32 InSizeX, int32 InSizeY, const TArray<FColor>& InImageData);
+
+	// Do some logging and trigger OnScreenshotTakenAndCompared
+	void OnComparisonComplete(bool bWasNew, bool bWasSimilar, double MaxLocalDifference, double GlobalDifference, FString ErrorMessage);
+
+	// Restore viewport size and restore original environment
+	void RestoreGameViewport();
+
+protected:
+
+	UPROPERTY()
+	class UCameraComponent* ScreenshotCamera;
+
+	UPROPERTY(EditAnywhere, BlueprintReadOnly, Category = "Screenshot", SimpleDisplay)
+	FAutomationScreenshotOptions ScreenshotOptions;
+
+	FIntPoint ViewportRestoreSize;
+
+#if (WITH_DEV_AUTOMATION_TESTS || WITH_PERF_AUTOMATION_TESTS)
+	TSharedPtr<FAutomationTestScreenshotEnvSetup> ScreenshotEnvSetup;
+#endif
+};
diff --git a/Engine/Source/Developer/FunctionalTesting/Private/AutomationBlueprintFunctionLibrary.cpp b/Engine/Source/Developer/FunctionalTesting/Private/AutomationBlueprintFunctionLibrary.cpp
--- a/Engine/Source/Developer/FunctionalTesting/Private/AutomationBlueprintFunctionLibrary.cpp
+++ b/Engine/Source/Developer/FunctionalTesting/Private/AutomationBlueprintFunctionLibrary.cpp
@@ -47,53 +47,134 @@ static TAutoConsoleVariable<int32> CVarAutomationScreenshotResolutionHeight(
 	TEXT("The height of automation screenshots."),
 	ECVF_Default);
 
-
 #if (WITH_DEV_AUTOMATION_TESTS || WITH_PERF_AUTOMATION_TESTS)
 
-class FConsoleVariableSwapper
+template<typename T>
+FConsoleVariableSwapperTempl<T>::FConsoleVariableSwapperTempl(FString InConsoleVariableName)
+	: bModified(false)
+	, ConsoleVariableName(InConsoleVariableName)
 {
-public:
-	FConsoleVariableSwapper(FString InConsoleVariableName)
-		: bModified(false)
-		, ConsoleVariableName(InConsoleVariableName)
+}
+
+template<typename T>
+void FConsoleVariableSwapperTempl<T>::Set(T Value)
+{
+	IConsoleVariable* ConsoleVariable = IConsoleManager::Get().FindConsoleVariable(*ConsoleVariableName);
+	if (ensure(ConsoleVariable))
 	{
+		if (bModified == false)
+		{
+			bModified = true;
+			OriginalValue = ConsoleVariable->GetInt();
+		}
+
+		ConsoleVariable->Set(Value);
 	}
+}
 
-	void Set(int32 Value)
+template<typename T>
+void FConsoleVariableSwapperTempl<T>::Restore()
+{
+	if (bModified)
 	{
 		IConsoleVariable* ConsoleVariable = IConsoleManager::Get().FindConsoleVariable(*ConsoleVariableName);
-		if ( ensure(ConsoleVariable) )
+		if (ensure(ConsoleVariable))
 		{
-			if ( bModified == false )
-			{
-				bModified = true;
-				OriginalValue = ConsoleVariable->GetInt();
-			}
+			ConsoleVariable->Set(OriginalValue);
+		}
 
-			ConsoleVariable->Set(Value);
+		bModified = false;
+	}
+}
+
+template<>
+void FConsoleVariableSwapperTempl<float>::Set(float Value)
+{
+	IConsoleVariable* ConsoleVariable = IConsoleManager::Get().FindConsoleVariable(*ConsoleVariableName);
+	if (ensure(ConsoleVariable))
+	{
+		if (bModified == false)
+		{
+			bModified = true;
+			OriginalValue = ConsoleVariable->GetFloat();
 		}
+
+		ConsoleVariable->Set(Value);
 	}
+}
+
+FAutomationTestScreenshotEnvSetup::FAutomationTestScreenshotEnvSetup()
+	: DefaultFeature_AntiAliasing(TEXT("r.DefaultFeature.AntiAliasing"))
+	, DefaultFeature_AutoExposure(TEXT("r.DefaultFeature.AutoExposure"))
+	, DefaultFeature_MotionBlur(TEXT("r.DefaultFeature.MotionBlur"))
+	, PostProcessAAQuality(TEXT("r.PostProcessAAQuality"))
+	, MotionBlurQuality(TEXT("r.MotionBlurQuality"))
+	, ScreenSpaceReflectionQuality(TEXT("r.SSR.Quality"))
+	, EyeAdaptationQuality(TEXT("r.EyeAdaptationQuality"))
+	, ContactShadows(TEXT("r.ContactShadows"))
+{
+}
 
-	void Restore()
+void FAutomationTestScreenshotEnvSetup::Setup(FAutomationScreenshotOptions& InOutOptions)
+{
+	check(IsInGameThread());
+
+	if (InOutOptions.bDisableNoisyRenderingFeatures)
 	{
-		if ( bModified )
+		DefaultFeature_AntiAliasing.Set(0);
+		DefaultFeature_AutoExposure.Set(0);
+		DefaultFeature_MotionBlur.Set(0);
+		PostProcessAAQuality.Set(0);
+		MotionBlurQuality.Set(0);
+		ScreenSpaceReflectionQuality.Set(0);
+		EyeAdaptationQuality.Set(0);
+		ContactShadows.Set(0);
+	}
+
+	InOutOptions.SetToleranceAmounts(InOutOptions.Tolerance);
+
+	if (UGameViewportClient* ViewportClient = GEngine->GameViewport)
+	{
+		static IConsoleVariable* ICVar = IConsoleManager::Get().FindConsoleVariable(FBufferVisualizationData::GetVisualizationTargetConsoleCommandName());
+		if (ICVar)
 		{
-			IConsoleVariable* ConsoleVariable = IConsoleManager::Get().FindConsoleVariable(*ConsoleVariableName);
-			if ( ensure(ConsoleVariable) )
+			if (ViewportClient->GetEngineShowFlags())
 			{
-				ConsoleVariable->Set(OriginalValue);
+				ViewportClient->GetEngineShowFlags()->SetVisualizeBuffer(InOutOptions.VisualizeBuffer == NAME_None ? false : true);
+				ViewportClient->GetEngineShowFlags()->SetTonemapper(InOutOptions.VisualizeBuffer == NAME_None ? true : false);
+				ICVar->Set(*InOutOptions.VisualizeBuffer.ToString());
 			}
-
-			bModified = false;
 		}
 	}
+}
 
-private:
-	bool bModified;
-	FString ConsoleVariableName;
-
-	int32 OriginalValue;
-};
+void FAutomationTestScreenshotEnvSetup::Restore()
+{
+	check(IsInGameThread());
+
+	DefaultFeature_AntiAliasing.Restore();
+	DefaultFeature_AutoExposure.Restore();
+	DefaultFeature_MotionBlur.Restore();
+	PostProcessAAQuality.Restore();
+	MotionBlurQuality.Restore();
+	ScreenSpaceReflectionQuality.Restore();
+	EyeAdaptationQuality.Restore();
+	ContactShadows.Restore();
+
+	if (UGameViewportClient* ViewportClient = GEngine->GameViewport)
+	{
+		static IConsoleVariable* ICVar = IConsoleManager::Get().FindConsoleVariable(FBufferVisualizationData::GetVisualizationTargetConsoleCommandName());
+		if (ICVar)
+		{
+			if (ViewportClient->GetEngineShowFlags())
+			{
+				ViewportClient->GetEngineShowFlags()->SetVisualizeBuffer(false);
+				ViewportClient->GetEngineShowFlags()->SetTonemapper(true);
+				ICVar->Set(TEXT(""));
+			}
+		}
+	}
+}
 
 class FAutomationScreenshotTaker
 {
@@ -102,74 +183,15 @@ public:
 		: World(InWorld)
 		, Name(InName)
 		, Options(InOptions)
-		, DefaultFeature_AntiAliasing(TEXT("r.DefaultFeature.AntiAliasing"))
-		, DefaultFeature_AutoExposure(TEXT("r.DefaultFeature.AutoExposure"))
-		, DefaultFeature_MotionBlur(TEXT("r.DefaultFeature.MotionBlur"))
-		, PostProcessAAQuality(TEXT("r.PostProcessAAQuality"))
-		, MotionBlurQuality(TEXT("r.MotionBlurQuality"))
-		, ScreenSpaceReflectionQuality(TEXT("r.SSR.Quality"))
-		, EyeAdaptationQuality(TEXT("r.EyeAdaptationQuality"))
-		, ContactShadows(TEXT("r.ContactShadows"))
 	{
 		GEngine->GameViewport->OnScreenshotCaptured().AddRaw(this, &FAutomationScreenshotTaker::GrabScreenShot);
 
-		check(IsInGameThread());
-
-		if ( Options.bDisableNoisyRenderingFeatures )
-		{
-			DefaultFeature_AntiAliasing.Set(0);
-			DefaultFeature_AutoExposure.Set(0);
-			DefaultFeature_MotionBlur.Set(0);
-			PostProcessAAQuality.Set(0);
-			MotionBlurQuality.Set(0);
-			ScreenSpaceReflectionQuality.Set(0);
-			EyeAdaptationQuality.Set(0);
-			ContactShadows.Set(0);
-		}
-
-		Options.SetToleranceAmounts(Options.Tolerance);
-
-		if ( UGameViewportClient* ViewportClient = GEngine->GameViewport )
-		{
-			static IConsoleVariable* ICVar = IConsoleManager::Get().FindConsoleVariable(FBufferVisualizationData::GetVisualizationTargetConsoleCommandName());
-			if ( ICVar )
-			{
-				if ( ViewportClient->GetEngineShowFlags() )
-				{
-					ViewportClient->GetEngineShowFlags()->SetVisualizeBuffer(InOptions.VisualizeBuffer == NAME_None ? false : true);
-					ViewportClient->GetEngineShowFlags()->SetTonemapper(InOptions.VisualizeBuffer == NAME_None ? true : false);
-					ICVar->Set(*InOptions.VisualizeBuffer.ToString());
-				}
-			}
-		}
+		EnvSetup.Setup(Options);
 	}
 
 	virtual ~FAutomationScreenshotTaker()
 	{
-		check(IsInGameThread());
-
-		DefaultFeature_AntiAliasing.Restore();
-		DefaultFeature_AutoExposure.Restore();
-		DefaultFeature_MotionBlur.Restore();
-		PostProcessAAQuality.Restore();
-		MotionBlurQuality.Restore();
-		ScreenSpaceReflectionQuality.Restore();
-		EyeAdaptationQuality.Restore();
-		ContactShadows.Restore();
-
-		if ( UGameViewportClient* ViewportClient = GEngine->GameViewport )
-		{
-			static IConsoleVariable* ICVar = IConsoleManager::Get().FindConsoleVariable(FBufferVisualizationData::GetVisualizationTargetConsoleCommandName());
-			if ( ICVar )
-			{
-				if ( ViewportClient->GetEngineShowFlags() )
-				{
-					ViewportClient->GetEngineShowFlags()->SetVisualizeBuffer(false);
-					ViewportClient->GetEngineShowFlags()->SetTonemapper(true);
-					ICVar->Set(TEXT(""));
-				}
-			}
-		}
+		EnvSetup.Restore();
 
 		GEngine->GameViewport->OnScreenshotCaptured().RemoveAll(this);
 
@@ -246,14 +268,7 @@ private:
 	FString	Name;
 	FAutomationScreenshotOptions Options;
 
-	FConsoleVariableSwapper DefaultFeature_AntiAliasing;
-	FConsoleVariableSwapper DefaultFeature_AutoExposure;
-	FConsoleVariableSwapper DefaultFeature_MotionBlur;
-	FConsoleVariableSwapper PostProcessAAQuality;
-	FConsoleVariableSwapper MotionBlurQuality;
-	FConsoleVariableSwapper ScreenSpaceReflectionQuality;
-	FConsoleVariableSwapper EyeAdaptationQuality;
-	FConsoleVariableSwapper ContactShadows;
+	FAutomationTestScreenshotEnvSetup EnvSetup;
 };
 
 #endif
@@ -278,24 +293,22 @@ void UAutomationBlueprintFunctionLibrary::FinishLoadingBeforeScreenshot()
 	IStreamingManager::Get().StreamAllResources(0.0f);
 }
 
-bool UAutomationBlueprintFunctionLibrary::TakeAutomationScreenshotInternal(UObject* WorldContextObject, const FString& Name, FAutomationScreenshotOptions Options)
+FIntPoint UAutomationBlueprintFunctionLibrary::GetAutomationScreenshotSize(const FAutomationScreenshotOptions& Options)
 {
-	UAutomationBlueprintFunctionLibrary::FinishLoadingBeforeScreenshot();
-
 	// Fallback resolution if all else fails for screenshots.
 	uint32 ResolutionX = 1280;
 	uint32 ResolutionY = 720;
 
 	// First get the default set for the project.
 	UAutomationTestSettings const* AutomationTestSettings = GetDefault<UAutomationTestSettings>();
-	if ( AutomationTestSettings->DefaultScreenshotResolution.GetMin() > 0 )
+	if (AutomationTestSettings->DefaultScreenshotResolution.GetMin() > 0)
 	{
 		ResolutionX = (uint32)AutomationTestSettings->DefaultScreenshotResolution.X;
 		ResolutionY = (uint32)AutomationTestSettings->DefaultScreenshotResolution.Y;
 	}
-	
+
 	// If there's an override resolution, use that instead.
-	if ( Options.Resolution.GetMin() > 0 )
+	if (Options.Resolution.GetMin() > 0)
 	{
 		ResolutionX = (uint32)Options.Resolution.X;
 		ResolutionY = (uint32)Options.Resolution.Y;
@@ -307,17 +320,26 @@ bool UAutomationBlueprintFunctionLibrary::TakeAutomationScreenshotInternal(UObje
 		int32 OverrideWidth = CVarAutomationScreenshotResolutionWidth.GetValueOnGameThread();
 		int32 OverrideHeight = CVarAutomationScreenshotResolutionHeight.GetValueOnGameThread();
 
-		if ( OverrideWidth > 0 )
+		if (OverrideWidth > 0)
 		{
 			ResolutionX = (uint32)OverrideWidth;
 		}
 
-		if ( OverrideHeight > 0 )
+		if (OverrideHeight > 0)
 		{
 			ResolutionY = (uint32)OverrideHeight;
 		}
 	}
 
+	return FIntPoint(ResolutionX, ResolutionY);
+}
+
+bool UAutomationBlueprintFunctionLibrary::TakeAutomationScreenshotInternal(UObject* WorldContextObject, const FString& Name, FAutomationScreenshotOptions Options)
+{
+	UAutomationBlueprintFunctionLibrary::FinishLoadingBeforeScreenshot();
+
+	FIntPoint ScreenshotRes = GetAutomationScreenshotSize(Options);
+
 #if (WITH_DEV_AUTOMATION_TESTS || WITH_PERF_AUTOMATION_TESTS)
 	FAutomationScreenshotTaker* TempObject = new FAutomationScreenshotTaker(WorldContextObject ? WorldContextObject->GetWorld() : nullptr, Name, Options);
 #endif
@@ -335,7 +357,7 @@ bool UAutomationBlueprintFunctionLibrary::TakeAutomationScreenshotInternal(UObje
 	{
 	    FHighResScreenshotConfig& Config = GetHighResScreenshotConfig();
 
-	    if ( Config.SetResolution(ResolutionX, ResolutionY, 1.0f) )
+	    if ( Config.SetResolution(ScreenshotRes.X, ScreenshotRes.Y, 1.0f) )
 	    {
 			if ( !GEngine->GameViewport->GetGameViewport()->TakeHighResScreenShot() )
 			{
diff --git a/Engine/Source/Developer/FunctionalTesting/Private/FunctionalUIScreenshotTest.cpp b/Engine/Source/Developer/FunctionalTesting/Private/FunctionalUIScreenshotTest.cpp
--- a/Engine/Source/Developer/FunctionalTesting/Private/FunctionalUIScreenshotTest.cpp
+++ b/Engine/Source/Developer/FunctionalTesting/Private/FunctionalUIScreenshotTest.cpp
@@ -10,27 +10,68 @@
 #include "Engine/Engine.h"
 #include "EngineGlobals.h"
 #include "Misc/AutomationTest.h"
+#include "Widgets/SViewport.h"
+#include "Engine/TextureRenderTarget2D.h"
+#include "Slate/SceneViewport.h"
+#include "Slate/WidgetRenderer.h"
 
 AFunctionalUIScreenshotTest::AFunctionalUIScreenshotTest( const FObjectInitializer& ObjectInitializer )
-	: AFunctionalTest(ObjectInitializer)
-	, ScreenshotOptions(EComparisonTolerance::Low)
+	: AScreenshotFunctionalTestBase(ObjectInitializer)
 {
-	ScreenshotCamera = CreateDefaultSubobject<UCameraComponent>(TEXT("Camera"));
-	ScreenshotCamera->SetupAttachment(RootComponent);
-
 	WidgetLocation = EWidgetTestAppearLocation::Viewport;
 }
 
+/**
+ * Get pixel format and color space of a backbuffer. Do nothing if the viewport doesn't
+ * render into backbuffer directly
+ * @InViewport - the viewport to get backbuffer from
+ * @OutPixelFormat - pixel format of the backbuffer
+ * @OutIsSRGB - whether the backbuffer stores pixels in sRGB space
+ */ 
+void GetBackbufferInfo(const FViewport* InViewport, EPixelFormat* OutPixelFormat, bool* OutIsSRGB)
+{
+	if (!InViewport->GetViewportRHI())
+	{
+		return;
+	}
+
+	ENQUEUE_RENDER_COMMAND(GetBackbufferFormatCmd)(
+		[InViewport, OutPixelFormat, OutIsSRGB](FRHICommandListImmediate& RHICmdList)
+	{
+		FViewportRHIRef ViewportRHI = InViewport->GetViewportRHI();
+		check(ViewportRHI.IsValid());
+		FTexture2DRHIRef BackbufferTexture = RHICmdList.GetViewportBackBuffer(ViewportRHI);
+		check(BackbufferTexture.IsValid());
+		*OutPixelFormat = BackbufferTexture->GetFormat();
+		*OutIsSRGB = (BackbufferTexture->GetFlags() & TexCreate_SRGB) == TexCreate_SRGB;
+	});
+	FlushRenderingCommands();
+}
+
 void AFunctionalUIScreenshotTest::PrepareTest()
 {
-	APlayerController* PlayerController = UGameplayStatics::GetPlayerController(GetWorld(), 0);
-	PlayerController->SetViewTarget(this, FViewTargetTransitionParams());
+	// Resize viewport to screenshot size
+	Super::PrepareTest();
 
-	// It's possible the defaults for certain tolerance levels have changed, so reset them on test start.
-	ScreenshotOptions.SetToleranceAmounts(ScreenshotOptions.Tolerance);
+	// Resize screenshot render target to have the same size as the game viewport. Also
+	// make sure they have the same data format (pixel format, color space, etc.) if possible
+	const FSceneViewport* GameViewport = GEngine->GameViewport->GetGameViewport();
+	FIntPoint ScreenshotSize = GameViewport->GetSizeXY();
+	EPixelFormat PixelFormat = PF_A2B10G10R10;
+	bool bIsSRGB = false;
+	GetBackbufferInfo(GameViewport, &PixelFormat, &bIsSRGB);
+
+	if (!ScreenshotRT)
+	{
+		ScreenshotRT = NewObject<UTextureRenderTarget2D>(this);
+	}
+	ScreenshotRT->ClearColor = FLinearColor::Transparent;
+	ScreenshotRT->InitCustomFormat(ScreenshotSize.X, ScreenshotSize.Y, PixelFormat, !bIsSRGB);
 
 	// Spawn the widget
+	APlayerController* PlayerController = UGameplayStatics::GetPlayerController(GetWorld(), 0);
 	SpawnedWidget = CreateWidget<UUserWidget>(PlayerController, WidgetClass);
+	check(SpawnedWidget)
 
 	if (WidgetLocation == EWidgetTestAppearLocation::Viewport)
 	{
@@ -41,83 +82,106 @@ void AFunctionalUIScreenshotTest::PrepareTest()
 		SpawnedWidget->AddToPlayerScreen();
 	}
 
-	UAutomationBlueprintFunctionLibrary::FinishLoadingBeforeScreenshot();
-
-	Super::PrepareTest();
-}
-
-bool AFunctionalUIScreenshotTest::IsReady_Implementation()
-{
-	if ( (GetWorld()->GetTimeSeconds() - StartTime) > ScreenshotOptions.Delay )
-	{
-		return ( GFrameNumber - RunFrame ) > 5;
-	}
-	
-	return false;
-}
-
-void AFunctionalUIScreenshotTest::StartTest()
-{
-	Super::StartTest();
-
 	if (SpawnedWidget)
 	{
 		SpawnedWidget->SetVisibility(ESlateVisibility::HitTestInvisible);
 	}
 
-	UAutomationBlueprintFunctionLibrary::TakeAutomationScreenshotOfUI_Immediate(this, GetName(), ScreenshotOptions);
-
-	FAutomationTestFramework::Get().OnScreenshotTakenAndCompared.AddUObject(this, &AFunctionalUIScreenshotTest::OnScreenshotTakenAndCompared);
+	UAutomationBlueprintFunctionLibrary::FinishLoadingBeforeScreenshot();
 }
 
 void AFunctionalUIScreenshotTest::OnScreenshotTakenAndCompared()
 {
-	FAutomationTestFramework::Get().OnScreenshotTakenAndCompared.RemoveAll(this);
-
 	if (SpawnedWidget)
 	{
 		SpawnedWidget->RemoveFromParent();
 	}
 
-	FinishTest(EFunctionalTestResult::Succeeded, TEXT(""));
+	// Restore viewport size and finish the test
+	Super::OnScreenshotTakenAndCompared();
 }
 
-#if WITH_EDITOR
+/**
+ * Source and destination textures need to have the same data format (pixel format, color space, etc.)
+ * @InViewport - backbuffer comes from this viewport
+ * @InOutScreenshotRT - copy destination
+ */
+void CopyBackbufferToScreenshotRT(const FViewport* InViewport, UTextureRenderTarget2D* InOutScreenshotRT)
+{
+	// Make sure rendering to the viewport backbuffer is finished
+	FlushRenderingCommands();
+
+	ENQUEUE_RENDER_COMMAND(CopyBackbufferCmd)(
+		[InViewport, InOutScreenshotRT](FRHICommandListImmediate& RHICmdList)
+	{
+		FViewportRHIRef ViewportRHI = InViewport->GetViewportRHI();
+		FTexture2DRHIRef BackbufferTexture = RHICmdList.GetViewportBackBuffer(ViewportRHI);
+		FTexture2DRHIRef ScreenshotRTTexture = InOutScreenshotRT->GetRenderTargetResource()->GetRenderTargetTexture();
+		check(BackbufferTexture->GetSizeXY() == ScreenshotRTTexture->GetSizeXY());
+		RHICmdList.CopyToResolveTarget(BackbufferTexture, ScreenshotRTTexture, true, FResolveParams());
+	});
+	FlushRenderingCommands();
+}
 
-bool AFunctionalUIScreenshotTest::CanEditChange(const UProperty* InProperty) const
+void ReadPixelsFromRT(UTextureRenderTarget2D* InRT, TArray<FColor>* OutPixels)
 {
-	bool bIsEditable = Super::CanEditChange(InProperty);
-	if ( bIsEditable && InProperty )
+	ENQUEUE_RENDER_COMMAND(ReadScreenshotRTCmd)(
+		[InRT, OutPixels](FRHICommandListImmediate& RHICmdList)
 	{
-		const FName PropertyName = InProperty->GetFName();
+		FTextureRenderTarget2DResource* RTResource =
+			static_cast<FTextureRenderTarget2DResource*>(InRT->GetRenderTargetResource());
+		RHICmdList.ReadSurfaceData(
+			RTResource->GetTextureRHI(),
+			FIntRect(0, 0, InRT->SizeX, InRT->SizeY),
+			*OutPixels,
+			FReadSurfaceDataFlags());
+	});
+	FlushRenderingCommands();
+}
 
-		if ( PropertyName == GET_MEMBER_NAME_CHECKED(FAutomationScreenshotOptions, ToleranceAmount) )
-		{
-			bIsEditable = ScreenshotOptions.Tolerance == EComparisonTolerance::Custom;
-		}
-		else if ( PropertyName == TEXT("ObservationPoint") )
+void AFunctionalUIScreenshotTest::RequestScreenshot()
+{
+	// Register a handler to UGameViewportClient::OnScreenshotCaptured
+	Super::RequestScreenshot();
+
+	UGameViewportClient* GameViewportClient = GEngine->GameViewport;
+	TSharedPtr<SViewport> ViewportWidget = GameViewportClient->GetGameViewportWidget();
+	FIntPoint ScreenshotSize = GameViewportClient->GetGameViewport()->GetSizeXY();
+	TArray<FColor> OutColorData;
+
+	if (ViewportWidget.IsValid())
+	{
+		// If SViewport renders directly to backbuffer (no separate RT), the scene will
+		// be black (only widgets will be rendered) if we just render SViewport. So we
+		// need to copy the scene to ScreenshotRT before rendering the test widget
+		if (ViewportWidget->ShouldRenderDirectly())
 		{
-			// You can't ever observe from anywhere but the camera on the screenshot test.
-			bIsEditable = false;
+			const FSceneViewport* GameViewport = GameViewportClient->GetGameViewport();
+			check(GameViewport);
+			CopyBackbufferToScreenshotRT(GameViewport, ScreenshotRT);
 		}
-	}
 
-	return bIsEditable;
-}
+		// Draw the game viewport (overlaid with the widget to screenshot) to our ScreenshotRT
+		TSharedPtr<FWidgetRenderer> WidgetRenderer = MakeShareable(new FWidgetRenderer(true, false));
+		check(WidgetRenderer.IsValid());
+		WidgetRenderer->DrawWidget(ScreenshotRT, ViewportWidget.ToSharedRef(), ScreenshotSize, 0.f);
+		FlushRenderingCommands();
 
-void AFunctionalUIScreenshotTest::PostEditChangeProperty(FPropertyChangedEvent& PropertyChangedEvent)
-{
-	Super::PostEditChangeProperty(PropertyChangedEvent);
+		// Possibly resolves the RT to a texture so we can read it back
+		ScreenshotRT->UpdateResourceImmediate(false);
 
-	if ( PropertyChangedEvent.Property )
-	{
-		const FName PropertyName = PropertyChangedEvent.Property->GetFName();
+		ReadPixelsFromRT(ScreenshotRT, &OutColorData);
+		check(OutColorData.Num() == ScreenshotSize.X * ScreenshotSize.Y);
 
-		if ( PropertyName == GET_MEMBER_NAME_CHECKED(FAutomationScreenshotOptions, Tolerance) )
+		// For UI, we only care about what the final image looks like. So don't compare alpha channel.
+		// In editor, scene is rendered into a PF_B8G8R8A8 RT and then copied over to the R10B10G10A2 swapchain back buffer and
+		// this copy ignores alpha. In game, however, scene is directly rendered into the back buffer and the alpha values are
+		// already meaningless at that stage.
+		for (int32 Idx = 0; Idx < OutColorData.Num(); ++Idx)
 		{
-			ScreenshotOptions.SetToleranceAmounts(ScreenshotOptions.Tolerance);
+			OutColorData[Idx].A = 0xff;
 		}
 	}
-}
 
-#endif
+	GameViewportClient->OnScreenshotCaptured().Broadcast(ScreenshotSize.X, ScreenshotSize.Y, OutColorData);
+}
diff --git a/Engine/Source/Developer/FunctionalTesting/Private/ScreenshotFunctionalTest.cpp b/Engine/Source/Developer/FunctionalTesting/Private/ScreenshotFunctionalTest.cpp
--- a/Engine/Source/Developer/FunctionalTesting/Private/ScreenshotFunctionalTest.cpp
+++ b/Engine/Source/Developer/FunctionalTesting/Private/ScreenshotFunctionalTest.cpp
@@ -12,30 +12,20 @@
 #include "Misc/AutomationTest.h"
 
 AScreenshotFunctionalTest::AScreenshotFunctionalTest( const FObjectInitializer& ObjectInitializer )
-	: AFunctionalTest(ObjectInitializer)
+	: AScreenshotFunctionalTestBase(ObjectInitializer)
 	, bCameraCutOnScreenshotPrep(false)
-	, ScreenshotOptions(EComparisonTolerance::Low)
 {
-	ScreenshotCamera = CreateDefaultSubobject<UCameraComponent>(TEXT("Camera"));
-	ScreenshotCamera->SetupAttachment(RootComponent);
 }
 
 void AScreenshotFunctionalTest::PrepareTest()
 {
-	APlayerController* PlayerController = UGameplayStatics::GetPlayerController(GetWorld(), 0);
-	if (PlayerController)
-	{
-		PlayerController->SetViewTarget(this, FViewTargetTransitionParams());
-	}
-
-	// It's possible the defaults for certain tolerance levels have changed, so reset them on test start.
-	ScreenshotOptions.SetToleranceAmounts(ScreenshotOptions.Tolerance);
-
 	Super::PrepareTest();
 
 	// Apply a camera cut if requested
 	if (bCameraCutOnScreenshotPrep)
 	{
+		APlayerController* PlayerController = UGameplayStatics::GetPlayerController(GetWorld(), 0);
+
 		if (PlayerController && PlayerController->PlayerCameraManager)
 		{
 			PlayerController->PlayerCameraManager->bGameCameraCutThisFrame = true;
@@ -45,70 +35,23 @@ void AScreenshotFunctionalTest::PrepareTest()
 			}
 		}
 	}
-}
 
-bool AScreenshotFunctionalTest::IsReady_Implementation()
-{
-	if ( (GetWorld()->GetTimeSeconds() - RunTime) > ScreenshotOptions.Delay )
-	{
-		return ( GFrameNumber - RunFrame ) > 5;
-	}
-	
-	return false;
+	UAutomationBlueprintFunctionLibrary::FinishLoadingBeforeScreenshot();
 }
 
-void AScreenshotFunctionalTest::StartTest()
+void AScreenshotFunctionalTest::RequestScreenshot()
 {
-	Super::StartTest();
-
-	UAutomationBlueprintFunctionLibrary::TakeAutomationScreenshotInternal(this, GetName(), ScreenshotOptions);
-
-	FAutomationTestFramework::Get().OnScreenshotTakenAndCompared.AddUObject(this, &AScreenshotFunctionalTest::OnScreenshotTakenAndCompared);
-}
-
-void AScreenshotFunctionalTest::OnScreenshotTakenAndCompared()
-{
-	FAutomationTestFramework::Get().OnScreenshotTakenAndCompared.RemoveAll(this);
-
-	FinishTest(EFunctionalTestResult::Succeeded, TEXT(""));
+	Super::RequestScreenshot();
+
+	// Screenshots in UE4 work in this way:
+	// 1. Call FScreenshotRequest::RequestScreenshot to ask the system to take a screenshot. The screenshot
+	//    will have the same resolution as the current viewport;
+	// 2. Register a callback to UGameViewportClient::OnScreenshotCaptured() delegate. The call back will be
+	//    called with screenshot pixel data when the shot is taken;
+	// 3. Wait till the next frame or call FSceneViewport::Invalidate to force a redraw. Screenshot is not
+	//    taken until next draw where UGameViewportClient::ProcessScreenshots or
+	//    FEditorViewportClient::ProcessScreenshots is called to read pixels back from the viewport. It also
+	//    trigger the callback function registered in step 2.
+	bool bShowUI = false;
+	FScreenshotRequest::RequestScreenshot(bShowUI);
 }
-
-#if WITH_EDITOR
-
-bool AScreenshotFunctionalTest::CanEditChange(const UProperty* InProperty) const
-{
-	bool bIsEditable = Super::CanEditChange(InProperty);
-	if ( bIsEditable && InProperty )
-	{
-		const FName PropertyName = InProperty->GetFName();
-
-		if ( PropertyName == GET_MEMBER_NAME_CHECKED(FAutomationScreenshotOptions, ToleranceAmount) )
-		{
-			bIsEditable = ScreenshotOptions.Tolerance == EComparisonTolerance::Custom;
-		}
-		else if ( PropertyName == TEXT("ObservationPoint") )
-		{
-			// You can't ever observe from anywhere but the camera on the screenshot test.
-			bIsEditable = false;
-		}
-	}
-
-	return bIsEditable;
-}
-
-void AScreenshotFunctionalTest::PostEditChangeProperty(FPropertyChangedEvent& PropertyChangedEvent)
-{
-	Super::PostEditChangeProperty(PropertyChangedEvent);
-
-	if ( PropertyChangedEvent.Property )
-	{
-		const FName PropertyName = PropertyChangedEvent.Property->GetFName();
-
-		if ( PropertyName == GET_MEMBER_NAME_CHECKED(FAutomationScreenshotOptions, Tolerance) )
-		{
-			ScreenshotOptions.SetToleranceAmounts(ScreenshotOptions.Tolerance);
-		}
-	}
-}
-
-#endif
diff --git a/Engine/Source/Developer/FunctionalTesting/Private/ScreenshotFunctionalTestBase.cpp b/Engine/Source/Developer/FunctionalTesting/Private/ScreenshotFunctionalTestBase.cpp
new file mode 100644
--- /dev/null
+++ b/Engine/Source/Developer/FunctionalTesting/Private/ScreenshotFunctionalTestBase.cpp
@@ -0,0 +1,222 @@
+// Copyright 1998-2017 Epic Games, Inc. All Rights Reserved.
+
+#include "ScreenshotFunctionalTestBase.h"
+
+#include "Engine/GameViewportClient.h"
+#include "AutomationBlueprintFunctionLibrary.h"
+#include "Camera/CameraComponent.h"
+#include "Camera/PlayerCameraManager.h"
+#include "Kismet/GameplayStatics.h"
+#include "Engine/Engine.h"
+#include "EngineGlobals.h"
+#include "Misc/AutomationTest.h"
+#include "Slate/SceneViewport.h"
+#include "RenderingThread.h"
+#include "Tests/AutomationCommon.h"
+#include "Logging/LogMacros.h"
+
+DEFINE_LOG_CATEGORY_STATIC(LogScreenshotFunctionalTest, Log, Log)
+
+AScreenshotFunctionalTestBase::AScreenshotFunctionalTestBase(const FObjectInitializer& ObjectInitializer)
+	: AFunctionalTest(ObjectInitializer)
+	, ScreenshotOptions(EComparisonTolerance::Low)
+{
+	ScreenshotCamera = CreateDefaultSubobject<UCameraComponent>(TEXT("Camera"));
+	ScreenshotCamera->SetupAttachment(RootComponent);
+
+#if (WITH_DEV_AUTOMATION_TESTS || WITH_PERF_AUTOMATION_TESTS)
+	ScreenshotEnvSetup = MakeShareable(new FAutomationTestScreenshotEnvSetup());
+#endif
+}
+
+void AScreenshotFunctionalTestBase::PrepareTest()
+{
+	Super::PrepareTest();
+
+	APlayerController* PlayerController = UGameplayStatics::GetPlayerController(GetWorld(), 0);
+	if (PlayerController)
+	{
+		PlayerController->SetViewTarget(this, FViewTargetTransitionParams());
+	}
+
+	PrepareForScreenshot();
+}
+
+bool AScreenshotFunctionalTestBase::IsReady_Implementation()
+{
+	if ((GetWorld()->GetTimeSeconds() - RunTime) > ScreenshotOptions.Delay)
+	{
+		return (GFrameNumber - RunFrame) > 5;
+	}
+
+	return false;
+}
+
+void AScreenshotFunctionalTestBase::StartTest()
+{
+	Super::StartTest();
+
+	FAutomationTestFramework::Get().OnScreenshotTakenAndCompared.AddUObject(this, &AScreenshotFunctionalTestBase::OnScreenshotTakenAndCompared);
+
+	RequestScreenshot();
+}
+
+void AScreenshotFunctionalTestBase::OnScreenshotTakenAndCompared()
+{
+	RestoreGameViewport();
+
+	FAutomationTestFramework::Get().OnScreenshotTakenAndCompared.RemoveAll(this);
+
+	FinishTest(EFunctionalTestResult::Succeeded, TEXT(""));
+}
+
+void AScreenshotFunctionalTestBase::PrepareForScreenshot()
+{
+	check(GEngine->GameViewport && GEngine->GameViewport->GetGameViewport());
+	check(IsInGameThread());
+
+	FlushRenderingCommands();
+
+	// PS4 requires fixed with/height for video output back buffers so
+	// we cannot adjust the size of viewports
+	if (!FPlatformProperties::HasFixedResolution())
+	{
+		FSceneViewport* GameViewport = GEngine->GameViewport->GetGameViewport();
+		ViewportRestoreSize = GameViewport->GetSize();
+		FIntPoint ScreenshotViewportSize = UAutomationBlueprintFunctionLibrary::GetAutomationScreenshotSize(ScreenshotOptions);
+		GameViewport->SetViewportSize(ScreenshotViewportSize.X, ScreenshotViewportSize.Y);
+	}
+
+#if (WITH_DEV_AUTOMATION_TESTS || WITH_PERF_AUTOMATION_TESTS)
+	ScreenshotEnvSetup->Setup(ScreenshotOptions);
+#endif
+}
+
+void AScreenshotFunctionalTestBase::OnScreenShotCaptured(int32 InSizeX, int32 InSizeY, const TArray<FColor>& InImageData)
+{
+	check(GEngine->GameViewport);
+
+	GEngine->GameViewport->OnScreenshotCaptured().RemoveAll(this);
+
+	FAutomationScreenshotData Data = AutomationCommon::BuildScreenshotData(GetWorld()->GetName(), GetName(), InSizeX, InSizeY);
+
+	// Copy the relevant data into the metadata for the screenshot.
+	Data.bHasComparisonRules = true;
+	Data.ToleranceRed = ScreenshotOptions.ToleranceAmount.Red;
+	Data.ToleranceGreen = ScreenshotOptions.ToleranceAmount.Green;
+	Data.ToleranceBlue = ScreenshotOptions.ToleranceAmount.Blue;
+	Data.ToleranceAlpha = ScreenshotOptions.ToleranceAmount.Alpha;
+	Data.ToleranceMinBrightness = ScreenshotOptions.ToleranceAmount.MinBrightness;
+	Data.ToleranceMaxBrightness = ScreenshotOptions.ToleranceAmount.MaxBrightness;
+	Data.bIgnoreAntiAliasing = ScreenshotOptions.bIgnoreAntiAliasing;
+	Data.bIgnoreColors = ScreenshotOptions.bIgnoreColors;
+	Data.MaximumLocalError = ScreenshotOptions.MaximumLocalError;
+	Data.MaximumGlobalError = ScreenshotOptions.MaximumGlobalError;
+
+	if (GIsAutomationTesting)
+	{
+		FAutomationTestFramework::Get().OnScreenshotCompared.AddUObject(this, &AScreenshotFunctionalTestBase::OnComparisonComplete);
+	}
+
+	FAutomationTestFramework::Get().OnScreenshotCaptured().ExecuteIfBound(InImageData, Data);
+
+	UE_LOG(LogScreenshotFunctionalTest, Log, TEXT("Screenshot captured as %s"), *Data.Path);
+}
+
+void AScreenshotFunctionalTestBase::RequestScreenshot()
+{
+	check(IsInGameThread());
+	check(GEngine->GameViewport);
+
+	// Make sure any screenshot request has been processed
+	FlushRenderingCommands();
+
+	UGameViewportClient* GameViewportClient = GEngine->GameViewport;
+	GameViewportClient->OnScreenshotCaptured().AddUObject(this, &AScreenshotFunctionalTestBase::OnScreenShotCaptured);
+}
+
+void AScreenshotFunctionalTestBase::OnComparisonComplete(bool bWasNew, bool bWasSimilar, double MaxLocalDifference, double GlobalDifference, FString ErrorMessage)
+{
+	FAutomationTestFramework::Get().OnScreenshotCompared.RemoveAll(this);
+
+	if (bWasNew)
+	{
+		UE_LOG(LogScreenshotFunctionalTest, Warning, TEXT("New Screenshot '%s' was discovered!  Please add a ground truth version of it."), *GetName());
+	}
+	else
+	{
+		if (bWasSimilar)
+		{
+			UE_LOG(LogScreenshotFunctionalTest, Display, TEXT("Screenshot '%s' was similar!  Global Difference = %f, Max Local Difference = %f"), *GetName(), GlobalDifference, MaxLocalDifference);
+		}
+		else
+		{
+			if (ErrorMessage.IsEmpty())
+			{
+				UE_LOG(LogScreenshotFunctionalTest, Error, TEXT("Screenshot '%s' test failed, Screnshots were different!  Global Difference = %f, Max Local Difference = %f"), *GetName(), GlobalDifference, MaxLocalDifference);
+			}
+			else
+			{
+				UE_LOG(LogScreenshotFunctionalTest, Error, TEXT("Screenshot '%s' test failed;  Error = %s"), *GetName(), *ErrorMessage);
+			}
+		}
+	}
+
+	FAutomationTestFramework::Get().NotifyScreenshotTakenAndCompared();
+}
+
+void AScreenshotFunctionalTestBase::RestoreGameViewport()
+{
+	check(GEngine->GameViewport && GEngine->GameViewport->GetGameViewport());
+	check(IsInGameThread());
+
+	if (!FPlatformProperties::HasFixedResolution())
+	{
+		FSceneViewport* GameViewport = GEngine->GameViewport->GetGameViewport();
+		GameViewport->SetViewportSize(ViewportRestoreSize.X, ViewportRestoreSize.Y);
+	}
+
+#if (WITH_DEV_AUTOMATION_TESTS || WITH_PERF_AUTOMATION_TESTS)
+	ScreenshotEnvSetup->Restore();
+#endif
+}
+
+#if WITH_EDITOR
+
+bool AScreenshotFunctionalTestBase::CanEditChange(const UProperty* InProperty) const
+{
+	bool bIsEditable = Super::CanEditChange(InProperty);
+	if (bIsEditable && InProperty)
+	{
+		const FName PropertyName = InProperty->GetFName();
+
+		if (PropertyName == GET_MEMBER_NAME_CHECKED(FAutomationScreenshotOptions, ToleranceAmount))
+		{
+			bIsEditable = ScreenshotOptions.Tolerance == EComparisonTolerance::Custom;
+		}
+		else if (PropertyName == TEXT("ObservationPoint"))
+		{
+			// You can't ever observe from anywhere but the camera on the screenshot test.
+			bIsEditable = false;
+		}
+	}
+
+	return bIsEditable;
+}
+
+void AScreenshotFunctionalTestBase::PostEditChangeProperty(FPropertyChangedEvent& PropertyChangedEvent)
+{
+	Super::PostEditChangeProperty(PropertyChangedEvent);
+
+	if (PropertyChangedEvent.Property)
+	{
+		const FName PropertyName = PropertyChangedEvent.Property->GetFName();
+
+		if (PropertyName == GET_MEMBER_NAME_CHECKED(FAutomationScreenshotOptions, Tolerance))
+		{
+			ScreenshotOptions.SetToleranceAmounts(ScreenshotOptions.Tolerance);
+		}
+	}
+}
+
+#endif
diff --git a/Engine/Source/Developer/FunctionalTesting/Public/AutomationBlueprintFunctionLibrary.h b/Engine/Source/Developer/FunctionalTesting/Public/AutomationBlueprintFunctionLibrary.h
--- a/Engine/Source/Developer/FunctionalTesting/Public/AutomationBlueprintFunctionLibrary.h
+++ b/Engine/Source/Developer/FunctionalTesting/Public/AutomationBlueprintFunctionLibrary.h
@@ -24,6 +24,8 @@ public:
 
 	static bool TakeAutomationScreenshotInternal(UObject* WorldContextObject, const FString& Name, FAutomationScreenshotOptions Options);
 
+	static FIntPoint GetAutomationScreenshotSize(const FAutomationScreenshotOptions& Options);
+
 	/**
 	 * Takes a screenshot of the game's viewport.  Does not capture any UI.
 	 */
@@ -83,3 +85,47 @@ public:
 	UFUNCTION(BlueprintPure, Category="Automation")
 	static FAutomationScreenshotOptions GetDefaultScreenshotOptionsForRendering(EComparisonTolerance Tolerance = EComparisonTolerance::Low, float Delay = 0.2);
 };
+
+#if (WITH_DEV_AUTOMATION_TESTS || WITH_PERF_AUTOMATION_TESTS)
+
+template<typename T>
+class FConsoleVariableSwapperTempl
+{
+public:
+	FConsoleVariableSwapperTempl(FString InConsoleVariableName);
+
+	void Set(T Value);
+
+	void Restore();
+
+private:
+	bool bModified;
+	FString ConsoleVariableName;
+
+	T OriginalValue;
+};
+
+class FAutomationTestScreenshotEnvSetup
+{
+public:
+	FAutomationTestScreenshotEnvSetup();
+
+	// Disable AA, auto-exposure, motion blur, contact shadow if InOutOptions.bDisableNoisyRenderingFeatures.
+	// Update screenshot comparison tolerance stored in InOutOptions.
+	// Set visualization buffer name if required.
+	void Setup(FAutomationScreenshotOptions& InOutOptions);
+
+	void Restore();
+
+private:
+	FConsoleVariableSwapperTempl<int32> DefaultFeature_AntiAliasing;
+	FConsoleVariableSwapperTempl<int32> DefaultFeature_AutoExposure;
+	FConsoleVariableSwapperTempl<int32> DefaultFeature_MotionBlur;
+	FConsoleVariableSwapperTempl<int32> PostProcessAAQuality;
+	FConsoleVariableSwapperTempl<int32> MotionBlurQuality;
+	FConsoleVariableSwapperTempl<int32> ScreenSpaceReflectionQuality;
+	FConsoleVariableSwapperTempl<int32> EyeAdaptationQuality;
+	FConsoleVariableSwapperTempl<int32> ContactShadows;
+};
+
+#endif
diff --git a/Engine/Source/Runtime/Engine/Private/TextureDerivedData.cpp b/Engine/Source/Runtime/Engine/Private/TextureDerivedData.cpp
--- a/Engine/Source/Runtime/Engine/Private/TextureDerivedData.cpp
+++ b/Engine/Source/Runtime/Engine/Private/TextureDerivedData.cpp
@@ -817,7 +817,7 @@ bool FTexturePlatformData::TryLoadMips(int32 FirstMipToLoad, void** OutMipData)
 				OutMipData[MipIndex - FirstMipToLoad] = FMemory::Malloc(Mip.BulkData.GetBulkDataSize());
 #if 0
 				checkSlow(!Mip.BulkData.GetFilename().EndsWith(TEXT(".ubulk"))); // We want to make sure that any non-streamed mips are coming from the texture asset file, and not from an external bulk file
-#else
+#elif !(WITH_DEV_AUTOMATION_TESTS || WITH_PERF_AUTOMATION_TESTS)
 				UE_CLOG(Mip.BulkData.GetFilename().EndsWith(TEXT(".ubulk")), LogTexture, Error, TEXT("Loading non-streamed mips from an external bulk file.  This is not desireable.  File %s"), *(Mip.BulkData.GetFilename() ) );
 #endif
 				
