Fix some warnings caused by 3647910

diff --git a/Engine/Source/Developer/FunctionalTesting/Classes/ScreenshotFunctionalTestBase.h b/Engine/Source/Developer/FunctionalTesting/Classes/ScreenshotFunctionalTestBase.h
--- a/Engine/Source/Developer/FunctionalTesting/Classes/ScreenshotFunctionalTestBase.h
+++ b/Engine/Source/Developer/FunctionalTesting/Classes/ScreenshotFunctionalTestBase.h
@@ -59,7 +59,7 @@ protected:
 
 protected:
 
-	UPROPERTY()
+	UPROPERTY(BlueprintReadOnly, Category = "Screenshot")
 	class UCameraComponent* ScreenshotCamera;
 
 	UPROPERTY(EditAnywhere, BlueprintReadOnly, Category = "Screenshot", SimpleDisplay)
diff --git a/Engine/Source/Developer/FunctionalTesting/Private/FunctionalUIScreenshotTest.cpp b/Engine/Source/Developer/FunctionalTesting/Private/FunctionalUIScreenshotTest.cpp
--- a/Engine/Source/Developer/FunctionalTesting/Private/FunctionalUIScreenshotTest.cpp
+++ b/Engine/Source/Developer/FunctionalTesting/Private/FunctionalUIScreenshotTest.cpp
@@ -71,19 +71,18 @@ void AFunctionalUIScreenshotTest::PrepareTest()
 	// Spawn the widget
 	APlayerController* PlayerController = UGameplayStatics::GetPlayerController(GetWorld(), 0);
 	SpawnedWidget = CreateWidget<UUserWidget>(PlayerController, WidgetClass);
-	check(SpawnedWidget)
-
-	if (WidgetLocation == EWidgetTestAppearLocation::Viewport)
-	{
-		SpawnedWidget->AddToViewport();
-	}
-	else
-	{
-		SpawnedWidget->AddToPlayerScreen();
-	}
-
+	
 	if (SpawnedWidget)
 	{
+		if (WidgetLocation == EWidgetTestAppearLocation::Viewport)
+		{
+			SpawnedWidget->AddToViewport();
+		}
+		else
+		{
+			SpawnedWidget->AddToPlayerScreen();
+		}
+
 		SpawnedWidget->SetVisibility(ESlateVisibility::HitTestInvisible);
 	}
 
