Fix Simple UI test

diff --git a/Engine/Source/Developer/FunctionalTesting/Private/FunctionalUIScreenshotTest.cpp b/Engine/Source/Developer/FunctionalTesting/Private/FunctionalUIScreenshotTest.cpp
--- a/Engine/Source/Developer/FunctionalTesting/Private/FunctionalUIScreenshotTest.cpp
+++ b/Engine/Source/Developer/FunctionalTesting/Private/FunctionalUIScreenshotTest.cpp
@@ -53,20 +53,27 @@ void AFunctionalUIScreenshotTest::PrepareTest()
 	// Resize viewport to screenshot size
 	Super::PrepareTest();
 
-	// Resize screenshot render target to have the same size as the game viewport. Also
-	// make sure they have the same data format (pixel format, color space, etc.) if possible
-	const FSceneViewport* GameViewport = GEngine->GameViewport->GetGameViewport();
-	FIntPoint ScreenshotSize = GameViewport->GetSizeXY();
-	EPixelFormat PixelFormat = PF_A2B10G10R10;
-	bool bIsSRGB = false;
-	GetBackbufferInfo(GameViewport, &PixelFormat, &bIsSRGB);
-
-	if (!ScreenshotRT)
+	TSharedPtr<SViewport> GameViewportWidget = GEngine->GameViewport->GetGameViewportWidget();
+	check(GameViewportWidget.IsValid());
+
+	// If render directly to backbuffer, just read from backbuffer
+	if (!GameViewportWidget->ShouldRenderDirectly())
 	{
-		ScreenshotRT = NewObject<UTextureRenderTarget2D>(this);
+		// Resize screenshot render target to have the same size as the game viewport. Also
+		// make sure they have the same data format (pixel format, color space, etc.) if possible
+		const FSceneViewport* GameViewport = GEngine->GameViewport->GetGameViewport();
+		FIntPoint ScreenshotSize = GameViewport->GetSizeXY();
+		EPixelFormat PixelFormat = PF_A2B10G10R10;
+		bool bIsSRGB = false;
+		GetBackbufferInfo(GameViewport, &PixelFormat, &bIsSRGB);
+
+		if (!ScreenshotRT)
+		{
+			ScreenshotRT = NewObject<UTextureRenderTarget2D>(this);
+		}
+		ScreenshotRT->ClearColor = FLinearColor::Transparent;
+		ScreenshotRT->InitCustomFormat(ScreenshotSize.X, ScreenshotSize.Y, PixelFormat, !bIsSRGB);
 	}
-	ScreenshotRT->ClearColor = FLinearColor::Transparent;
-	ScreenshotRT->InitCustomFormat(ScreenshotSize.X, ScreenshotSize.Y, PixelFormat, !bIsSRGB);
 
 	// Spawn the widget
 	APlayerController* PlayerController = UGameplayStatics::GetPlayerController(GetWorld(), 0);
@@ -80,6 +87,8 @@ void AFunctionalUIScreenshotTest::PrepareTest()
 		}
 		else
 		{
+			// Add to the game viewport and restrain the widget within
+			// owning player's sub-rect
 			SpawnedWidget->AddToPlayerScreen();
 		}
 
@@ -101,23 +110,25 @@ void AFunctionalUIScreenshotTest::OnScreenshotTakenAndCompared()
 }
 
 /**
- * Source and destination textures need to have the same data format (pixel format, color space, etc.)
+ * Read all pixels from backbuffer
  * @InViewport - backbuffer comes from this viewport
- * @InOutScreenshotRT - copy destination
+ * @OutPixels
  */
-void CopyBackbufferToScreenshotRT(const FViewport* InViewport, UTextureRenderTarget2D* InOutScreenshotRT)
+void ReadBackbuffer(const FViewport* InViewport, TArray<FColor>* OutPixels)
 {
 	// Make sure rendering to the viewport backbuffer is finished
 	FlushRenderingCommands();
 
 	ENQUEUE_RENDER_COMMAND(CopyBackbufferCmd)(
-		[InViewport, InOutScreenshotRT](FRHICommandListImmediate& RHICmdList)
+		[InViewport, OutPixels](FRHICommandListImmediate& RHICmdList)
 	{
 		FViewportRHIRef ViewportRHI = InViewport->GetViewportRHI();
 		FTexture2DRHIRef BackbufferTexture = RHICmdList.GetViewportBackBuffer(ViewportRHI);
-		FTexture2DRHIRef ScreenshotRTTexture = InOutScreenshotRT->GetRenderTargetResource()->GetRenderTargetTexture();
-		check(BackbufferTexture->GetSizeXY() == ScreenshotRTTexture->GetSizeXY());
-		RHICmdList.CopyToResolveTarget(BackbufferTexture, ScreenshotRTTexture, true, FResolveParams());
+		RHICmdList.ReadSurfaceData(
+			BackbufferTexture,
+			FIntRect(0, 0, BackbufferTexture->GetSizeX(), BackbufferTexture->GetSizeY()),
+			*OutPixels,
+			FReadSurfaceDataFlags());
 	});
 	FlushRenderingCommands();
 }
@@ -150,37 +161,31 @@ void AFunctionalUIScreenshotTest::RequestScreenshot()
 
 	if (ViewportWidget.IsValid())
 	{
-		// If SViewport renders directly to backbuffer (no separate RT), the scene will
-		// be black (only widgets will be rendered) if we just render SViewport. So we
-		// need to copy the scene to ScreenshotRT before rendering the test widget
 		if (ViewportWidget->ShouldRenderDirectly())
 		{
 			const FSceneViewport* GameViewport = GameViewportClient->GetGameViewport();
-			check(GameViewport);
-			CopyBackbufferToScreenshotRT(GameViewport, ScreenshotRT);
+			ReadBackbuffer(GameViewport, &OutColorData);
+		}
+		else
+		{
+			// Draw the game viewport (overlaid with the widget to screenshot) to our ScreenshotRT.
+			// Need to do this manually because the game viewport doesn't have a valid FViewportRHIRef
+			// when rendering to a separate render target
+			TSharedPtr<FWidgetRenderer> WidgetRenderer = MakeShareable(new FWidgetRenderer(true, false));
+			check(WidgetRenderer.IsValid());
+			WidgetRenderer->DrawWidget(ScreenshotRT, ViewportWidget.ToSharedRef(), ScreenshotSize, 0.f);
+			FlushRenderingCommands();
+
+			ReadPixelsFromRT(ScreenshotRT, &OutColorData);
 		}
-
-		// Draw the game viewport (overlaid with the widget to screenshot) to our ScreenshotRT
-		TSharedPtr<FWidgetRenderer> WidgetRenderer = MakeShareable(new FWidgetRenderer(true, false));
-		check(WidgetRenderer.IsValid());
-		WidgetRenderer->DrawWidget(ScreenshotRT, ViewportWidget.ToSharedRef(), ScreenshotSize, 0.f);
-		FlushRenderingCommands();
-
-		// Possibly resolves the RT to a texture so we can read it back
-		ScreenshotRT->UpdateResourceImmediate(false);
-
-		ReadPixelsFromRT(ScreenshotRT, &OutColorData);
-		check(OutColorData.Num() == ScreenshotSize.X * ScreenshotSize.Y);
 
 		// For UI, we only care about what the final image looks like. So don't compare alpha channel.
-		// In editor, scene is rendered into a PF_B8G8R8A8 RT and then copied over to the R10B10G10A2 swapchain back buffer and
-		// this copy ignores alpha. In game, however, scene is directly rendered into the back buffer and the alpha values are
-		// already meaningless at that stage.
 		for (int32 Idx = 0; Idx < OutColorData.Num(); ++Idx)
 		{
 			OutColorData[Idx].A = 0xff;
 		}
 	}
 
+	check(OutColorData.Num() == ScreenshotSize.X * ScreenshotSize.Y);
 	GameViewportClient->OnScreenshotCaptured().Broadcast(ScreenshotSize.X, ScreenshotSize.Y, OutColorData);
 }
