Merging //UE4/Dev-Main to Release-Staging-4.18 (//UE4/Release-Staging-4.18)

diff --git a/Engine/Build/Build.version b/Engine/Build/Build.version
--- a/Engine/Build/Build.version
+++ b/Engine/Build/Build.version
@@ -1,6 +1,6 @@
 {
 	"MajorVersion": 4,
-	"MinorVersion": 18,
+	"MinorVersion": 19,
 	"PatchVersion": 0,
 	"Changelist": 0,
 	"CompatibleChangelist": 0,
diff --git a/Engine/Documentation/Source/Programming/Development/BuildGraph/ScriptAnatomy/Tasks/PredefinedTasks/BuildGraphScriptPredefinedTasks.INT.udn b/Engine/Documentation/Source/Programming/Development/BuildGraph/ScriptAnatomy/Tasks/PredefinedTasks/BuildGraphScriptPredefinedTasks.INT.udn
--- a/Engine/Documentation/Source/Programming/Development/BuildGraph/ScriptAnatomy/Tasks/PredefinedTasks/BuildGraphScriptPredefinedTasks.INT.udn
+++ b/Engine/Documentation/Source/Programming/Development/BuildGraph/ScriptAnatomy/Tasks/PredefinedTasks/BuildGraphScriptPredefinedTasks.INT.udn
@@ -2,7 +2,7 @@ Availability: NoPublish
 Title: BuildGraph Predefined Tasks
 Crumbs: %ROOT%, Programming, Programming/Development, Programming/Development/BuildGraph, Programming/Development/BuildGraph/BuildGraphScriptTasks
 Description: This is a procedurally generated markdown page.
-version: 4.18
+version: 4.19
 parent:Programming/Development/BuildGraph/BuildGraphScriptTasks
 
 ### AgeStore
diff --git a/Engine/Documentation/Source/Programming/UnrealBuildSystem/Configuration/BuildConfigProperties/BuildConfigProperties.INT.udn b/Engine/Documentation/Source/Programming/UnrealBuildSystem/Configuration/BuildConfigProperties/BuildConfigProperties.INT.udn
--- a/Engine/Documentation/Source/Programming/UnrealBuildSystem/Configuration/BuildConfigProperties/BuildConfigProperties.INT.udn
+++ b/Engine/Documentation/Source/Programming/UnrealBuildSystem/Configuration/BuildConfigProperties/BuildConfigProperties.INT.udn
@@ -2,7 +2,7 @@ Availability: NoPublish
 Title: Build Configuration Properties Page
 Crumbs:
 Description: This is a procedurally generated markdown page.
-Version: 4.18
+Version: 4.19
 
 ### XboxOnePlatform
 
diff --git a/Engine/Documentation/Source/Programming/UnrealBuildSystem/ModuleFiles/ModuleFilesProperties/ModuleFilesProperties.INT.udn b/Engine/Documentation/Source/Programming/UnrealBuildSystem/ModuleFiles/ModuleFilesProperties/ModuleFilesProperties.INT.udn
--- a/Engine/Documentation/Source/Programming/UnrealBuildSystem/ModuleFiles/ModuleFilesProperties/ModuleFilesProperties.INT.udn
+++ b/Engine/Documentation/Source/Programming/UnrealBuildSystem/ModuleFiles/ModuleFilesProperties/ModuleFilesProperties.INT.udn
@@ -2,7 +2,7 @@ Availability: NoPublish
 Title: Build Configuration Properties Page
 Crumbs:
 Description: This is a procedurally generated markdown page.
-Version: 4.18
+Version: 4.19
 
 ### Read/Write Properties
 $ Type (ModuleType): Type of module
diff --git a/Engine/Documentation/Source/Programming/UnrealBuildSystem/TargetFiles/TargetFilesProperties/TargetFilesProperties.INT.udn b/Engine/Documentation/Source/Programming/UnrealBuildSystem/TargetFiles/TargetFilesProperties/TargetFilesProperties.INT.udn
--- a/Engine/Documentation/Source/Programming/UnrealBuildSystem/TargetFiles/TargetFilesProperties/TargetFilesProperties.INT.udn
+++ b/Engine/Documentation/Source/Programming/UnrealBuildSystem/TargetFiles/TargetFilesProperties/TargetFilesProperties.INT.udn
@@ -2,7 +2,7 @@ Availability: NoPublish
 Title: Build Configuration Properties Page
 Crumbs:
 Description: This is a procedurally generated markdown page.
-Version: 4.18
+Version: 4.19
 
 ### Read-Only Properties
 
diff --git a/Engine/Plugins/Tests/EditorTests/Source/EditorTests/Private/UnrealEd/LightingTests.cpp b/Engine/Plugins/Tests/EditorTests/Source/EditorTests/Private/UnrealEd/LightingTests.cpp
--- a/Engine/Plugins/Tests/EditorTests/Source/EditorTests/Private/UnrealEd/LightingTests.cpp
+++ b/Engine/Plugins/Tests/EditorTests/Source/EditorTests/Private/UnrealEd/LightingTests.cpp
@@ -16,6 +16,8 @@
 
 // Misc
 #include "ScopedTransaction.h"
+#include "Engine/Selection.h"
+#include "EditorModeManager.h"
 
 #define POINTLIGHT_TRANSFORM			FTransform()
 #define POINT_LIGHT_UPDATED_LOCATION	FVector(50.f, 50.f, 50.f)
@@ -196,26 +198,26 @@ bool FLightingPromotionModifyProperties::RunTest(const FString& Parameters)
 }
 
 
-/**
+/*
+/ **
 * Lighting Promotion Test - Duplicate/Copy Paste a point light.
-*/
+* /
 IMPLEMENT_SIMPLE_AUTOMATION_TEST(FLightingPromotionDuplicationTest, "System.Promotion.Editor.Lighting.Duplicate and Copy Paste", EAutomationTestFlags::EditorContext | EAutomationTestFlags::EngineFilter)
 
 bool FLightingPromotionDuplicationTest::RunTest(const FString& Parameters)
 {
-	//** SETUP **//
 	UWorld* World = FAutomationEditorCommonUtils::CreateNewMap();
 	ULevel* CurrentLevel = World->GetCurrentLevel();
 	// Test Summary
 	AddInfo(TEXT("Duplicate and Copy Paste\n- Duplicates a point light.\n- Copies and Pastes a point light."));
 
+	GLevelEditorModeTools().ActivateDefaultMode();
+
 	if (!LightingTestHelpers::DoesActorExistInTheLevel(CurrentLevel, TEXT("PointLight"), APointLight::StaticClass()))
 	{
-		//** TEST **//
 		// Add a point light to the level.
 		APointLight* PointLight = Cast<APointLight>(GEditor->AddActor(World->GetCurrentLevel(), APointLight::StaticClass(), FTransform()));
 
-		//** Verify **//
 		int32 NumberOfPointLights = 0;
 		// Count the number of point lights in the level.
 		for ( TActorIterator<APointLight> It(World); It; ++It )
@@ -228,11 +230,17 @@ bool FLightingPromotionDuplicationTest::RunTest(const FString& Parameters)
 
 		// Deselect all and then Select the light
 		LightingTestHelpers::SelectActorInLevel(PointLight);
+
+		TestEqual(TEXT("Selected actor count"), GEditor->GetSelectedActorCount(), 1);
+
+		TestEqual<APointLight*>(TEXT("PointLight is selected"), PointLight, GEditor->GetSelectedActors()->GetTop<APointLight>());
+
 		// Copy and Paste.
-		GEngine->Exec(World, TEXT("EDIT COPY"));
-		GEngine->Exec(World, TEXT("EDIT PASTE"));
+		GEditor->CopySelectedActorsToClipboard(World, false);
+
+		FText TransDescription = FText::FromString("Automation Test Copy/Paste");
+		GEditor->PasteSelectedActorsFromClipboard(World, TransDescription, EPasteTo::PT_OriginalLocation);
 
-		//** Verify **//
 		NumberOfPointLights = 0;
 		// Count the number of point lights in the level.
 		for ( TActorIterator<APointLight> It(World); It; ++It )
@@ -245,10 +253,10 @@ bool FLightingPromotionDuplicationTest::RunTest(const FString& Parameters)
 
 		// Deselect all and then select a light
 		LightingTestHelpers::SelectActorInLevel(PointLight);
+
 		// Duplicate the light
 		GEngine->Exec(World, TEXT("DUPLICATE"));
 
-		//** Verify **//
 		NumberOfPointLights = 0;
 		// Count the number of point lights in the level.
 		for ( TActorIterator<APointLight> It(World); It; ++It )
@@ -264,7 +272,7 @@ bool FLightingPromotionDuplicationTest::RunTest(const FString& Parameters)
 
 	AddError(TEXT("A point light already exists in this level which would dirty the test results."));
 	return false;
-}
+}*/
 
 
 //////////////////////////////////////////////////////////////////////////
diff --git a/Engine/Source/Developer/FunctionalTesting/Classes/FunctionalUIScreenshotTest.h b/Engine/Source/Developer/FunctionalTesting/Classes/FunctionalUIScreenshotTest.h
--- a/Engine/Source/Developer/FunctionalTesting/Classes/FunctionalUIScreenshotTest.h
+++ b/Engine/Source/Developer/FunctionalTesting/Classes/FunctionalUIScreenshotTest.h
@@ -4,12 +4,14 @@
 
 #include "CoreMinimal.h"
 #include "UObject/ObjectMacros.h"
-#include "FunctionalTest.h"
+#include "ScreenshotFunctionalTestBase.h"
 #include "AutomationScreenshotOptions.h"
 #include "Blueprint/UserWidget.h"
 
 #include "FunctionalUIScreenshotTest.generated.h"
 
+class UTextureRenderTarget2D;
+
 UENUM()
 enum class EWidgetTestAppearLocation
 {
@@ -21,24 +23,19 @@ enum class EWidgetTestAppearLocation
  * 
  */
 UCLASS(Blueprintable, MinimalAPI)
-class AFunctionalUIScreenshotTest : public AFunctionalTest
+class AFunctionalUIScreenshotTest : public AScreenshotFunctionalTestBase
 {
 	GENERATED_BODY()
 
 public:
 	AFunctionalUIScreenshotTest(const FObjectInitializer& ObjectInitializer);
 
-#if WITH_EDITOR
-	virtual bool CanEditChange(const UProperty* InProperty) const override;
-	virtual void PostEditChangeProperty(FPropertyChangedEvent& PropertyChangedEvent) override;
-#endif
-
 protected:
 	virtual void PrepareTest() override;
-	virtual bool IsReady_Implementation() override;
-	virtual void StartTest() override;
 
-	void OnScreenshotTakenAndCompared();
+	virtual void OnScreenshotTakenAndCompared() override;
+
+	virtual void RequestScreenshot() override;
 	
 protected:
 
@@ -50,10 +47,7 @@ protected:
 
 	UPROPERTY(EditAnywhere, Category = "UI")
 	EWidgetTestAppearLocation WidgetLocation;
-	
-	UPROPERTY()
-	class UCameraComponent* ScreenshotCamera;
 
-	UPROPERTY(EditAnywhere, Category="Screenshot", SimpleDisplay)
-	FAutomationScreenshotOptions ScreenshotOptions;
+	UPROPERTY(Transient, DuplicateTransient)
+	UTextureRenderTarget2D* ScreenshotRT;
 };
diff --git a/Engine/Source/Developer/FunctionalTesting/Classes/ScreenshotFunctionalTest.h b/Engine/Source/Developer/FunctionalTesting/Classes/ScreenshotFunctionalTest.h
--- a/Engine/Source/Developer/FunctionalTesting/Classes/ScreenshotFunctionalTest.h
+++ b/Engine/Source/Developer/FunctionalTesting/Classes/ScreenshotFunctionalTest.h
@@ -4,43 +4,30 @@
 
 #include "CoreMinimal.h"
 #include "UObject/ObjectMacros.h"
-#include "FunctionalTest.h"
+#include "ScreenshotFunctionalTestBase.h"
 #include "AutomationScreenshotOptions.h"
 
 #include "ScreenshotFunctionalTest.generated.h"
 
+class FAutomationTestScreenshotEnvSetup;
+
 /**
- * 
+ * No UI
  */
 UCLASS(Blueprintable)
-class FUNCTIONALTESTING_API AScreenshotFunctionalTest : public AFunctionalTest
+class FUNCTIONALTESTING_API AScreenshotFunctionalTest : public AScreenshotFunctionalTestBase
 {
 	GENERATED_BODY()
 
 public:
 	AScreenshotFunctionalTest(const FObjectInitializer& ObjectInitializer);
 
-#if WITH_EDITOR
-	virtual bool CanEditChange(const UProperty* InProperty) const override;
-	virtual void PostEditChangeProperty(FPropertyChangedEvent& PropertyChangedEvent) override;
-#endif
-
 	// Tests relying on temporal effects can force a camera cut to flush stale data
 	UPROPERTY(EditAnywhere, BlueprintReadWrite, Category = "Camera", SimpleDisplay)
 	bool bCameraCutOnScreenshotPrep;
 
 protected:
 	virtual void PrepareTest() override;
-	virtual bool IsReady_Implementation() override;
-	virtual void StartTest() override;
-
-	void OnScreenshotTakenAndCompared();
-	
-protected:
-	
-	UPROPERTY(EditAnywhere, BlueprintReadWrite, Category = "Camera", SimpleDisplay)
-	class UCameraComponent* ScreenshotCamera;
 
-	UPROPERTY(EditAnywhere, BlueprintReadOnly, Category="Screenshot", SimpleDisplay)
-	FAutomationScreenshotOptions ScreenshotOptions;
+	virtual void RequestScreenshot() override;
 };
diff --git a/Engine/Source/Developer/FunctionalTesting/Classes/ScreenshotFunctionalTestBase.h b/Engine/Source/Developer/FunctionalTesting/Classes/ScreenshotFunctionalTestBase.h
new file mode 100644
--- /dev/null
+++ b/Engine/Source/Developer/FunctionalTesting/Classes/ScreenshotFunctionalTestBase.h
@@ -0,0 +1,73 @@
+// Copyright 1998-2017 Epic Games, Inc. All Rights Reserved.
+
+#pragma once
+
+#include "CoreMinimal.h"
+#include "UObject/ObjectMacros.h"
+#include "FunctionalTest.h"
+#include "AutomationScreenshotOptions.h"
+
+#include "ScreenshotFunctionalTestBase.generated.h"
+
+class FAutomationTestScreenshotEnvSetup;
+
+/**
+* Base class for screenshot functional test
+*/
+UCLASS(Blueprintable, abstract)
+class FUNCTIONALTESTING_API AScreenshotFunctionalTestBase : public AFunctionalTest
+{
+	GENERATED_BODY()
+
+public:
+	AScreenshotFunctionalTestBase(const FObjectInitializer& ObjectInitializer);
+
+#if WITH_EDITOR
+	virtual bool CanEditChange(const UProperty* InProperty) const override;
+	virtual void PostEditChangeProperty(FPropertyChangedEvent& PropertyChangedEvent) override;
+#endif
+
+protected:
+	// Set player view target to screenshot camera and call PrepareForScreenshot
+	virtual void PrepareTest() override;
+
+	// Handle screenshot delay
+	virtual bool IsReady_Implementation() override;
+
+	// Register OnScreenshotTakenAndCompared and call RequestScreenshot
+	virtual void StartTest() override;
+
+	// Call RestoreViewport and finish this test
+	virtual void OnScreenshotTakenAndCompared();
+
+	// Resize viewport to screenshot size (if possible) and set up screenshot environment (disable AA, etc.)
+	void PrepareForScreenshot();
+
+	// Doesn't actually request in base class. It simply register OnScreenshotCaptured
+	virtual void RequestScreenshot();
+
+	// Pass screenshot pixels and meta data to FAutomationTestFramework. Register
+	// OnComparisonComplete which will be called the automation test system when
+	// screenshot comparison is complete
+	void OnScreenShotCaptured(int32 InSizeX, int32 InSizeY, const TArray<FColor>& InImageData);
+
+	// Do some logging and trigger OnScreenshotTakenAndCompared
+	void OnComparisonComplete(bool bWasNew, bool bWasSimilar, double MaxLocalDifference, double GlobalDifference, FString ErrorMessage);
+
+	// Restore viewport size and restore original environment
+	void RestoreGameViewport();
+
+protected:
+
+	UPROPERTY(BlueprintReadOnly, Category = "Screenshot")
+	class UCameraComponent* ScreenshotCamera;
+
+	UPROPERTY(EditAnywhere, BlueprintReadOnly, Category = "Screenshot", SimpleDisplay)
+	FAutomationScreenshotOptions ScreenshotOptions;
+
+	FIntPoint ViewportRestoreSize;
+
+#if (WITH_DEV_AUTOMATION_TESTS || WITH_PERF_AUTOMATION_TESTS)
+	TSharedPtr<FAutomationTestScreenshotEnvSetup> ScreenshotEnvSetup;
+#endif
+};
diff --git a/Engine/Source/Developer/FunctionalTesting/Private/AutomationBlueprintFunctionLibrary.cpp b/Engine/Source/Developer/FunctionalTesting/Private/AutomationBlueprintFunctionLibrary.cpp
--- a/Engine/Source/Developer/FunctionalTesting/Private/AutomationBlueprintFunctionLibrary.cpp
+++ b/Engine/Source/Developer/FunctionalTesting/Private/AutomationBlueprintFunctionLibrary.cpp
@@ -47,53 +47,134 @@ static TAutoConsoleVariable<int32> CVarAutomationScreenshotResolutionHeight(
 	TEXT("The height of automation screenshots."),
 	ECVF_Default);
 
-
 #if (WITH_DEV_AUTOMATION_TESTS || WITH_PERF_AUTOMATION_TESTS)
 
-class FConsoleVariableSwapper
+template<typename T>
+FConsoleVariableSwapperTempl<T>::FConsoleVariableSwapperTempl(FString InConsoleVariableName)
+	: bModified(false)
+	, ConsoleVariableName(InConsoleVariableName)
 {
-public:
-	FConsoleVariableSwapper(FString InConsoleVariableName)
-		: bModified(false)
-		, ConsoleVariableName(InConsoleVariableName)
+}
+
+template<typename T>
+void FConsoleVariableSwapperTempl<T>::Set(T Value)
+{
+	IConsoleVariable* ConsoleVariable = IConsoleManager::Get().FindConsoleVariable(*ConsoleVariableName);
+	if (ensure(ConsoleVariable))
 	{
+		if (bModified == false)
+		{
+			bModified = true;
+			OriginalValue = ConsoleVariable->GetInt();
+		}
+
+		ConsoleVariable->Set(Value);
 	}
+}
 
-	void Set(int32 Value)
+template<typename T>
+void FConsoleVariableSwapperTempl<T>::Restore()
+{
+	if (bModified)
 	{
 		IConsoleVariable* ConsoleVariable = IConsoleManager::Get().FindConsoleVariable(*ConsoleVariableName);
-		if ( ensure(ConsoleVariable) )
+		if (ensure(ConsoleVariable))
 		{
-			if ( bModified == false )
-			{
-				bModified = true;
-				OriginalValue = ConsoleVariable->GetInt();
-			}
+			ConsoleVariable->Set(OriginalValue);
+		}
 
-			ConsoleVariable->Set(Value);
+		bModified = false;
+	}
+}
+
+template<>
+void FConsoleVariableSwapperTempl<float>::Set(float Value)
+{
+	IConsoleVariable* ConsoleVariable = IConsoleManager::Get().FindConsoleVariable(*ConsoleVariableName);
+	if (ensure(ConsoleVariable))
+	{
+		if (bModified == false)
+		{
+			bModified = true;
+			OriginalValue = ConsoleVariable->GetFloat();
 		}
+
+		ConsoleVariable->Set(Value);
 	}
+}
+
+FAutomationTestScreenshotEnvSetup::FAutomationTestScreenshotEnvSetup()
+	: DefaultFeature_AntiAliasing(TEXT("r.DefaultFeature.AntiAliasing"))
+	, DefaultFeature_AutoExposure(TEXT("r.DefaultFeature.AutoExposure"))
+	, DefaultFeature_MotionBlur(TEXT("r.DefaultFeature.MotionBlur"))
+	, PostProcessAAQuality(TEXT("r.PostProcessAAQuality"))
+	, MotionBlurQuality(TEXT("r.MotionBlurQuality"))
+	, ScreenSpaceReflectionQuality(TEXT("r.SSR.Quality"))
+	, EyeAdaptationQuality(TEXT("r.EyeAdaptationQuality"))
+	, ContactShadows(TEXT("r.ContactShadows"))
+{
+}
 
-	void Restore()
+void FAutomationTestScreenshotEnvSetup::Setup(FAutomationScreenshotOptions& InOutOptions)
+{
+	check(IsInGameThread());
+
+	if (InOutOptions.bDisableNoisyRenderingFeatures)
 	{
-		if ( bModified )
+		DefaultFeature_AntiAliasing.Set(0);
+		DefaultFeature_AutoExposure.Set(0);
+		DefaultFeature_MotionBlur.Set(0);
+		PostProcessAAQuality.Set(0);
+		MotionBlurQuality.Set(0);
+		ScreenSpaceReflectionQuality.Set(0);
+		EyeAdaptationQuality.Set(0);
+		ContactShadows.Set(0);
+	}
+
+	InOutOptions.SetToleranceAmounts(InOutOptions.Tolerance);
+
+	if (UGameViewportClient* ViewportClient = GEngine->GameViewport)
+	{
+		static IConsoleVariable* ICVar = IConsoleManager::Get().FindConsoleVariable(FBufferVisualizationData::GetVisualizationTargetConsoleCommandName());
+		if (ICVar)
 		{
-			IConsoleVariable* ConsoleVariable = IConsoleManager::Get().FindConsoleVariable(*ConsoleVariableName);
-			if ( ensure(ConsoleVariable) )
+			if (ViewportClient->GetEngineShowFlags())
 			{
-				ConsoleVariable->Set(OriginalValue);
+				ViewportClient->GetEngineShowFlags()->SetVisualizeBuffer(InOutOptions.VisualizeBuffer == NAME_None ? false : true);
+				ViewportClient->GetEngineShowFlags()->SetTonemapper(InOutOptions.VisualizeBuffer == NAME_None ? true : false);
+				ICVar->Set(*InOutOptions.VisualizeBuffer.ToString());
 			}
-
-			bModified = false;
 		}
 	}
+}
 
-private:
-	bool bModified;
-	FString ConsoleVariableName;
-
-	int32 OriginalValue;
-};
+void FAutomationTestScreenshotEnvSetup::Restore()
+{
+	check(IsInGameThread());
+
+	DefaultFeature_AntiAliasing.Restore();
+	DefaultFeature_AutoExposure.Restore();
+	DefaultFeature_MotionBlur.Restore();
+	PostProcessAAQuality.Restore();
+	MotionBlurQuality.Restore();
+	ScreenSpaceReflectionQuality.Restore();
+	EyeAdaptationQuality.Restore();
+	ContactShadows.Restore();
+
+	if (UGameViewportClient* ViewportClient = GEngine->GameViewport)
+	{
+		static IConsoleVariable* ICVar = IConsoleManager::Get().FindConsoleVariable(FBufferVisualizationData::GetVisualizationTargetConsoleCommandName());
+		if (ICVar)
+		{
+			if (ViewportClient->GetEngineShowFlags())
+			{
+				ViewportClient->GetEngineShowFlags()->SetVisualizeBuffer(false);
+				ViewportClient->GetEngineShowFlags()->SetTonemapper(true);
+				ICVar->Set(TEXT(""));
+			}
+		}
+	}
+}
 
 class FAutomationScreenshotTaker
 {
@@ -102,74 +183,15 @@ public:
 		: World(InWorld)
 		, Name(InName)
 		, Options(InOptions)
-		, DefaultFeature_AntiAliasing(TEXT("r.DefaultFeature.AntiAliasing"))
-		, DefaultFeature_AutoExposure(TEXT("r.DefaultFeature.AutoExposure"))
-		, DefaultFeature_MotionBlur(TEXT("r.DefaultFeature.MotionBlur"))
-		, PostProcessAAQuality(TEXT("r.PostProcessAAQuality"))
-		, MotionBlurQuality(TEXT("r.MotionBlurQuality"))
-		, ScreenSpaceReflectionQuality(TEXT("r.SSR.Quality"))
-		, EyeAdaptationQuality(TEXT("r.EyeAdaptationQuality"))
-		, ContactShadows(TEXT("r.ContactShadows"))
 	{
 		GEngine->GameViewport->OnScreenshotCaptured().AddRaw(this, &FAutomationScreenshotTaker::GrabScreenShot);
 
-		check(IsInGameThread());
-
-		if ( Options.bDisableNoisyRenderingFeatures )
-		{
-			DefaultFeature_AntiAliasing.Set(0);
-			DefaultFeature_AutoExposure.Set(0);
-			DefaultFeature_MotionBlur.Set(0);
-			PostProcessAAQuality.Set(0);
-			MotionBlurQuality.Set(0);
-			ScreenSpaceReflectionQuality.Set(0);
-			EyeAdaptationQuality.Set(0);
-			ContactShadows.Set(0);
-		}
-
-		Options.SetToleranceAmounts(Options.Tolerance);
-
-		if ( UGameViewportClient* ViewportClient = GEngine->GameViewport )
-		{
-			static IConsoleVariable* ICVar = IConsoleManager::Get().FindConsoleVariable(FBufferVisualizationData::GetVisualizationTargetConsoleCommandName());
-			if ( ICVar )
-			{
-				if ( ViewportClient->GetEngineShowFlags() )
-				{
-					ViewportClient->GetEngineShowFlags()->SetVisualizeBuffer(InOptions.VisualizeBuffer == NAME_None ? false : true);
-					ViewportClient->GetEngineShowFlags()->SetTonemapper(InOptions.VisualizeBuffer == NAME_None ? true : false);
-					ICVar->Set(*InOptions.VisualizeBuffer.ToString());
-				}
-			}
-		}
+		EnvSetup.Setup(Options);
 	}
 
 	virtual ~FAutomationScreenshotTaker()
 	{
-		check(IsInGameThread());
-
-		DefaultFeature_AntiAliasing.Restore();
-		DefaultFeature_AutoExposure.Restore();
-		DefaultFeature_MotionBlur.Restore();
-		PostProcessAAQuality.Restore();
-		MotionBlurQuality.Restore();
-		ScreenSpaceReflectionQuality.Restore();
-		EyeAdaptationQuality.Restore();
-		ContactShadows.Restore();
-
-		if ( UGameViewportClient* ViewportClient = GEngine->GameViewport )
-		{
-			static IConsoleVariable* ICVar = IConsoleManager::Get().FindConsoleVariable(FBufferVisualizationData::GetVisualizationTargetConsoleCommandName());
-			if ( ICVar )
-			{
-				if ( ViewportClient->GetEngineShowFlags() )
-				{
-					ViewportClient->GetEngineShowFlags()->SetVisualizeBuffer(false);
-					ViewportClient->GetEngineShowFlags()->SetTonemapper(true);
-					ICVar->Set(TEXT(""));
-				}
-			}
-		}
+		EnvSetup.Restore();
 
 		GEngine->GameViewport->OnScreenshotCaptured().RemoveAll(this);
 
@@ -246,14 +268,7 @@ private:
 	FString	Name;
 	FAutomationScreenshotOptions Options;
 
-	FConsoleVariableSwapper DefaultFeature_AntiAliasing;
-	FConsoleVariableSwapper DefaultFeature_AutoExposure;
-	FConsoleVariableSwapper DefaultFeature_MotionBlur;
-	FConsoleVariableSwapper PostProcessAAQuality;
-	FConsoleVariableSwapper MotionBlurQuality;
-	FConsoleVariableSwapper ScreenSpaceReflectionQuality;
-	FConsoleVariableSwapper EyeAdaptationQuality;
-	FConsoleVariableSwapper ContactShadows;
+	FAutomationTestScreenshotEnvSetup EnvSetup;
 };
 
 #endif
@@ -278,24 +293,22 @@ void UAutomationBlueprintFunctionLibrary::FinishLoadingBeforeScreenshot()
 	IStreamingManager::Get().StreamAllResources(0.0f);
 }
 
-bool UAutomationBlueprintFunctionLibrary::TakeAutomationScreenshotInternal(UObject* WorldContextObject, const FString& Name, FAutomationScreenshotOptions Options)
+FIntPoint UAutomationBlueprintFunctionLibrary::GetAutomationScreenshotSize(const FAutomationScreenshotOptions& Options)
 {
-	UAutomationBlueprintFunctionLibrary::FinishLoadingBeforeScreenshot();
-
 	// Fallback resolution if all else fails for screenshots.
 	uint32 ResolutionX = 1280;
 	uint32 ResolutionY = 720;
 
 	// First get the default set for the project.
 	UAutomationTestSettings const* AutomationTestSettings = GetDefault<UAutomationTestSettings>();
-	if ( AutomationTestSettings->DefaultScreenshotResolution.GetMin() > 0 )
+	if (AutomationTestSettings->DefaultScreenshotResolution.GetMin() > 0)
 	{
 		ResolutionX = (uint32)AutomationTestSettings->DefaultScreenshotResolution.X;
 		ResolutionY = (uint32)AutomationTestSettings->DefaultScreenshotResolution.Y;
 	}
-	
+
 	// If there's an override resolution, use that instead.
-	if ( Options.Resolution.GetMin() > 0 )
+	if (Options.Resolution.GetMin() > 0)
 	{
 		ResolutionX = (uint32)Options.Resolution.X;
 		ResolutionY = (uint32)Options.Resolution.Y;
@@ -307,17 +320,26 @@ bool UAutomationBlueprintFunctionLibrary::TakeAutomationScreenshotInternal(UObje
 		int32 OverrideWidth = CVarAutomationScreenshotResolutionWidth.GetValueOnGameThread();
 		int32 OverrideHeight = CVarAutomationScreenshotResolutionHeight.GetValueOnGameThread();
 
-		if ( OverrideWidth > 0 )
+		if (OverrideWidth > 0)
 		{
 			ResolutionX = (uint32)OverrideWidth;
 		}
 
-		if ( OverrideHeight > 0 )
+		if (OverrideHeight > 0)
 		{
 			ResolutionY = (uint32)OverrideHeight;
 		}
 	}
 
+	return FIntPoint(ResolutionX, ResolutionY);
+}
+
+bool UAutomationBlueprintFunctionLibrary::TakeAutomationScreenshotInternal(UObject* WorldContextObject, const FString& Name, FAutomationScreenshotOptions Options)
+{
+	UAutomationBlueprintFunctionLibrary::FinishLoadingBeforeScreenshot();
+
+	FIntPoint ScreenshotRes = GetAutomationScreenshotSize(Options);
+
 #if (WITH_DEV_AUTOMATION_TESTS || WITH_PERF_AUTOMATION_TESTS)
 	FAutomationScreenshotTaker* TempObject = new FAutomationScreenshotTaker(WorldContextObject ? WorldContextObject->GetWorld() : nullptr, Name, Options);
 #endif
@@ -335,7 +357,7 @@ bool UAutomationBlueprintFunctionLibrary::TakeAutomationScreenshotInternal(UObje
 	{
 	    FHighResScreenshotConfig& Config = GetHighResScreenshotConfig();
 
-	    if ( Config.SetResolution(ResolutionX, ResolutionY, 1.0f) )
+	    if ( Config.SetResolution(ScreenshotRes.X, ScreenshotRes.Y, 1.0f) )
 	    {
 			if ( !GEngine->GameViewport->GetGameViewport()->TakeHighResScreenShot() )
 			{
diff --git a/Engine/Source/Developer/FunctionalTesting/Private/FunctionalUIScreenshotTest.cpp b/Engine/Source/Developer/FunctionalTesting/Private/FunctionalUIScreenshotTest.cpp
--- a/Engine/Source/Developer/FunctionalTesting/Private/FunctionalUIScreenshotTest.cpp
+++ b/Engine/Source/Developer/FunctionalTesting/Private/FunctionalUIScreenshotTest.cpp
@@ -10,114 +10,182 @@
 #include "Engine/Engine.h"
 #include "EngineGlobals.h"
 #include "Misc/AutomationTest.h"
+#include "Widgets/SViewport.h"
+#include "Engine/TextureRenderTarget2D.h"
+#include "Slate/SceneViewport.h"
+#include "Slate/WidgetRenderer.h"
 
 AFunctionalUIScreenshotTest::AFunctionalUIScreenshotTest( const FObjectInitializer& ObjectInitializer )
-	: AFunctionalTest(ObjectInitializer)
-	, ScreenshotOptions(EComparisonTolerance::Low)
+	: AScreenshotFunctionalTestBase(ObjectInitializer)
 {
-	ScreenshotCamera = CreateDefaultSubobject<UCameraComponent>(TEXT("Camera"));
-	ScreenshotCamera->SetupAttachment(RootComponent);
-
 	WidgetLocation = EWidgetTestAppearLocation::Viewport;
 }
 
-void AFunctionalUIScreenshotTest::PrepareTest()
+/**
+ * Get pixel format and color space of a backbuffer. Do nothing if the viewport doesn't
+ * render into backbuffer directly
+ * @InViewport - the viewport to get backbuffer from
+ * @OutPixelFormat - pixel format of the backbuffer
+ * @OutIsSRGB - whether the backbuffer stores pixels in sRGB space
+ */ 
+void GetBackbufferInfo(const FViewport* InViewport, EPixelFormat* OutPixelFormat, bool* OutIsSRGB)
 {
-	APlayerController* PlayerController = UGameplayStatics::GetPlayerController(GetWorld(), 0);
-	PlayerController->SetViewTarget(this, FViewTargetTransitionParams());
-
-	// It's possible the defaults for certain tolerance levels have changed, so reset them on test start.
-	ScreenshotOptions.SetToleranceAmounts(ScreenshotOptions.Tolerance);
-
-	// Spawn the widget
-	SpawnedWidget = CreateWidget<UUserWidget>(PlayerController, WidgetClass);
-
-	if (WidgetLocation == EWidgetTestAppearLocation::Viewport)
-	{
-		SpawnedWidget->AddToViewport();
-	}
-	else
+	if (!InViewport->GetViewportRHI())
 	{
-		SpawnedWidget->AddToPlayerScreen();
+		return;
 	}
 
-	UAutomationBlueprintFunctionLibrary::FinishLoadingBeforeScreenshot();
-
-	Super::PrepareTest();
+	ENQUEUE_RENDER_COMMAND(GetBackbufferFormatCmd)(
+		[InViewport, OutPixelFormat, OutIsSRGB](FRHICommandListImmediate& RHICmdList)
+	{
+		FViewportRHIRef ViewportRHI = InViewport->GetViewportRHI();
+		check(ViewportRHI.IsValid());
+		FTexture2DRHIRef BackbufferTexture = RHICmdList.GetViewportBackBuffer(ViewportRHI);
+		check(BackbufferTexture.IsValid());
+		*OutPixelFormat = BackbufferTexture->GetFormat();
+		*OutIsSRGB = (BackbufferTexture->GetFlags() & TexCreate_SRGB) == TexCreate_SRGB;
+	});
+	FlushRenderingCommands();
 }
 
-bool AFunctionalUIScreenshotTest::IsReady_Implementation()
+void AFunctionalUIScreenshotTest::PrepareTest()
 {
-	if ( (GetWorld()->GetTimeSeconds() - StartTime) > ScreenshotOptions.Delay )
+	// Resize viewport to screenshot size
+	Super::PrepareTest();
+
+	TSharedPtr<SViewport> GameViewportWidget = GEngine->GameViewport->GetGameViewportWidget();
+	check(GameViewportWidget.IsValid());
+
+	// If render directly to backbuffer, just read from backbuffer
+	if (!GameViewportWidget->ShouldRenderDirectly())
 	{
-		return ( GFrameNumber - RunFrame ) > 5;
+		// Resize screenshot render target to have the same size as the game viewport. Also
+		// make sure they have the same data format (pixel format, color space, etc.) if possible
+		const FSceneViewport* GameViewport = GEngine->GameViewport->GetGameViewport();
+		FIntPoint ScreenshotSize = GameViewport->GetSizeXY();
+		EPixelFormat PixelFormat = PF_A2B10G10R10;
+		bool bIsSRGB = false;
+		GetBackbufferInfo(GameViewport, &PixelFormat, &bIsSRGB);
+
+		if (!ScreenshotRT)
+		{
+			ScreenshotRT = NewObject<UTextureRenderTarget2D>(this);
+		}
+		ScreenshotRT->ClearColor = FLinearColor::Transparent;
+		ScreenshotRT->InitCustomFormat(ScreenshotSize.X, ScreenshotSize.Y, PixelFormat, !bIsSRGB);
 	}
-	
-	return false;
-}
-
-void AFunctionalUIScreenshotTest::StartTest()
-{
-	Super::StartTest();
 
+	// Spawn the widget
+	APlayerController* PlayerController = UGameplayStatics::GetPlayerController(GetWorld(), 0);
+	SpawnedWidget = CreateWidget<UUserWidget>(PlayerController, WidgetClass);
+	
 	if (SpawnedWidget)
 	{
+		if (WidgetLocation == EWidgetTestAppearLocation::Viewport)
+		{
+			SpawnedWidget->AddToViewport();
+		}
+		else
+		{
+			// Add to the game viewport and restrain the widget within
+			// owning player's sub-rect
+			SpawnedWidget->AddToPlayerScreen();
+		}
+
 		SpawnedWidget->SetVisibility(ESlateVisibility::HitTestInvisible);
 	}
 
-	UAutomationBlueprintFunctionLibrary::TakeAutomationScreenshotOfUI_Immediate(this, GetName(), ScreenshotOptions);
-
-	FAutomationTestFramework::Get().OnScreenshotTakenAndCompared.AddUObject(this, &AFunctionalUIScreenshotTest::OnScreenshotTakenAndCompared);
+	UAutomationBlueprintFunctionLibrary::FinishLoadingBeforeScreenshot();
 }
 
 void AFunctionalUIScreenshotTest::OnScreenshotTakenAndCompared()
 {
-	FAutomationTestFramework::Get().OnScreenshotTakenAndCompared.RemoveAll(this);
-
 	if (SpawnedWidget)
 	{
 		SpawnedWidget->RemoveFromParent();
 	}
 
-	FinishTest(EFunctionalTestResult::Succeeded, TEXT(""));
+	// Restore viewport size and finish the test
+	Super::OnScreenshotTakenAndCompared();
 }
 
-#if WITH_EDITOR
-
-bool AFunctionalUIScreenshotTest::CanEditChange(const UProperty* InProperty) const
+/**
+ * Read all pixels from backbuffer
+ * @InViewport - backbuffer comes from this viewport
+ * @OutPixels
+ */
+void ReadBackbuffer(const FViewport* InViewport, TArray<FColor>* OutPixels)
 {
-	bool bIsEditable = Super::CanEditChange(InProperty);
-	if ( bIsEditable && InProperty )
-	{
-		const FName PropertyName = InProperty->GetFName();
+	// Make sure rendering to the viewport backbuffer is finished
+	FlushRenderingCommands();
 
-		if ( PropertyName == GET_MEMBER_NAME_CHECKED(FAutomationScreenshotOptions, ToleranceAmount) )
-		{
-			bIsEditable = ScreenshotOptions.Tolerance == EComparisonTolerance::Custom;
-		}
-		else if ( PropertyName == TEXT("ObservationPoint") )
-		{
-			// You can't ever observe from anywhere but the camera on the screenshot test.
-			bIsEditable = false;
-		}
-	}
+	ENQUEUE_RENDER_COMMAND(CopyBackbufferCmd)(
+		[InViewport, OutPixels](FRHICommandListImmediate& RHICmdList)
+	{
+		FViewportRHIRef ViewportRHI = InViewport->GetViewportRHI();
+		FTexture2DRHIRef BackbufferTexture = RHICmdList.GetViewportBackBuffer(ViewportRHI);
+		RHICmdList.ReadSurfaceData(
+			BackbufferTexture,
+			FIntRect(0, 0, BackbufferTexture->GetSizeX(), BackbufferTexture->GetSizeY()),
+			*OutPixels,
+			FReadSurfaceDataFlags());
+	});
+	FlushRenderingCommands();
+}
 
-	return bIsEditable;
+void ReadPixelsFromRT(UTextureRenderTarget2D* InRT, TArray<FColor>* OutPixels)
+{
+	ENQUEUE_RENDER_COMMAND(ReadScreenshotRTCmd)(
+		[InRT, OutPixels](FRHICommandListImmediate& RHICmdList)
+	{
+		FTextureRenderTarget2DResource* RTResource =
+			static_cast<FTextureRenderTarget2DResource*>(InRT->GetRenderTargetResource());
+		RHICmdList.ReadSurfaceData(
+			RTResource->GetTextureRHI(),
+			FIntRect(0, 0, InRT->SizeX, InRT->SizeY),
+			*OutPixels,
+			FReadSurfaceDataFlags());
+	});
+	FlushRenderingCommands();
 }
 
-void AFunctionalUIScreenshotTest::PostEditChangeProperty(FPropertyChangedEvent& PropertyChangedEvent)
+void AFunctionalUIScreenshotTest::RequestScreenshot()
 {
-	Super::PostEditChangeProperty(PropertyChangedEvent);
+	// Register a handler to UGameViewportClient::OnScreenshotCaptured
+	Super::RequestScreenshot();
+
+	UGameViewportClient* GameViewportClient = GEngine->GameViewport;
+	TSharedPtr<SViewport> ViewportWidget = GameViewportClient->GetGameViewportWidget();
+	FIntPoint ScreenshotSize = GameViewportClient->GetGameViewport()->GetSizeXY();
+	TArray<FColor> OutColorData;
 
-	if ( PropertyChangedEvent.Property )
+	if (ViewportWidget.IsValid())
 	{
-		const FName PropertyName = PropertyChangedEvent.Property->GetFName();
+		if (ViewportWidget->ShouldRenderDirectly())
+		{
+			const FSceneViewport* GameViewport = GameViewportClient->GetGameViewport();
+			ReadBackbuffer(GameViewport, &OutColorData);
+		}
+		else
+		{
+			// Draw the game viewport (overlaid with the widget to screenshot) to our ScreenshotRT.
+			// Need to do this manually because the game viewport doesn't have a valid FViewportRHIRef
+			// when rendering to a separate render target
+			TSharedPtr<FWidgetRenderer> WidgetRenderer = MakeShareable(new FWidgetRenderer(true, false));
+			check(WidgetRenderer.IsValid());
+			WidgetRenderer->DrawWidget(ScreenshotRT, ViewportWidget.ToSharedRef(), ScreenshotSize, 0.f);
+			FlushRenderingCommands();
+
+			ReadPixelsFromRT(ScreenshotRT, &OutColorData);
+		}
 
-		if ( PropertyName == GET_MEMBER_NAME_CHECKED(FAutomationScreenshotOptions, Tolerance) )
+		// For UI, we only care about what the final image looks like. So don't compare alpha channel.
+		for (int32 Idx = 0; Idx < OutColorData.Num(); ++Idx)
 		{
-			ScreenshotOptions.SetToleranceAmounts(ScreenshotOptions.Tolerance);
+			OutColorData[Idx].A = 0xff;
 		}
 	}
-}
 
-#endif
+	check(OutColorData.Num() == ScreenshotSize.X * ScreenshotSize.Y);
+	GameViewportClient->OnScreenshotCaptured().Broadcast(ScreenshotSize.X, ScreenshotSize.Y, OutColorData);
+}
diff --git a/Engine/Source/Developer/FunctionalTesting/Private/ScreenshotFunctionalTest.cpp b/Engine/Source/Developer/FunctionalTesting/Private/ScreenshotFunctionalTest.cpp
--- a/Engine/Source/Developer/FunctionalTesting/Private/ScreenshotFunctionalTest.cpp
+++ b/Engine/Source/Developer/FunctionalTesting/Private/ScreenshotFunctionalTest.cpp
@@ -12,30 +12,20 @@
 #include "Misc/AutomationTest.h"
 
 AScreenshotFunctionalTest::AScreenshotFunctionalTest( const FObjectInitializer& ObjectInitializer )
-	: AFunctionalTest(ObjectInitializer)
+	: AScreenshotFunctionalTestBase(ObjectInitializer)
 	, bCameraCutOnScreenshotPrep(false)
-	, ScreenshotOptions(EComparisonTolerance::Low)
 {
-	ScreenshotCamera = CreateDefaultSubobject<UCameraComponent>(TEXT("Camera"));
-	ScreenshotCamera->SetupAttachment(RootComponent);
 }
 
 void AScreenshotFunctionalTest::PrepareTest()
 {
-	APlayerController* PlayerController = UGameplayStatics::GetPlayerController(GetWorld(), 0);
-	if (PlayerController)
-	{
-		PlayerController->SetViewTarget(this, FViewTargetTransitionParams());
-	}
-
-	// It's possible the defaults for certain tolerance levels have changed, so reset them on test start.
-	ScreenshotOptions.SetToleranceAmounts(ScreenshotOptions.Tolerance);
-
 	Super::PrepareTest();
 
 	// Apply a camera cut if requested
 	if (bCameraCutOnScreenshotPrep)
 	{
+		APlayerController* PlayerController = UGameplayStatics::GetPlayerController(GetWorld(), 0);
+
 		if (PlayerController && PlayerController->PlayerCameraManager)
 		{
 			PlayerController->PlayerCameraManager->bGameCameraCutThisFrame = true;
@@ -45,70 +35,23 @@ void AScreenshotFunctionalTest::PrepareTest()
 			}
 		}
 	}
-}
 
-bool AScreenshotFunctionalTest::IsReady_Implementation()
-{
-	if ( (GetWorld()->GetTimeSeconds() - RunTime) > ScreenshotOptions.Delay )
-	{
-		return ( GFrameNumber - RunFrame ) > 5;
-	}
-	
-	return false;
+	UAutomationBlueprintFunctionLibrary::FinishLoadingBeforeScreenshot();
 }
 
-void AScreenshotFunctionalTest::StartTest()
+void AScreenshotFunctionalTest::RequestScreenshot()
 {
-	Super::StartTest();
-
-	UAutomationBlueprintFunctionLibrary::TakeAutomationScreenshotInternal(this, GetName(), ScreenshotOptions);
-
-	FAutomationTestFramework::Get().OnScreenshotTakenAndCompared.AddUObject(this, &AScreenshotFunctionalTest::OnScreenshotTakenAndCompared);
-}
-
-void AScreenshotFunctionalTest::OnScreenshotTakenAndCompared()
-{
-	FAutomationTestFramework::Get().OnScreenshotTakenAndCompared.RemoveAll(this);
-
-	FinishTest(EFunctionalTestResult::Succeeded, TEXT(""));
+	Super::RequestScreenshot();
+
+	// Screenshots in UE4 work in this way:
+	// 1. Call FScreenshotRequest::RequestScreenshot to ask the system to take a screenshot. The screenshot
+	//    will have the same resolution as the current viewport;
+	// 2. Register a callback to UGameViewportClient::OnScreenshotCaptured() delegate. The call back will be
+	//    called with screenshot pixel data when the shot is taken;
+	// 3. Wait till the next frame or call FSceneViewport::Invalidate to force a redraw. Screenshot is not
+	//    taken until next draw where UGameViewportClient::ProcessScreenshots or
+	//    FEditorViewportClient::ProcessScreenshots is called to read pixels back from the viewport. It also
+	//    trigger the callback function registered in step 2.
+	bool bShowUI = false;
+	FScreenshotRequest::RequestScreenshot(bShowUI);
 }
-
-#if WITH_EDITOR
-
-bool AScreenshotFunctionalTest::CanEditChange(const UProperty* InProperty) const
-{
-	bool bIsEditable = Super::CanEditChange(InProperty);
-	if ( bIsEditable && InProperty )
-	{
-		const FName PropertyName = InProperty->GetFName();
-
-		if ( PropertyName == GET_MEMBER_NAME_CHECKED(FAutomationScreenshotOptions, ToleranceAmount) )
-		{
-			bIsEditable = ScreenshotOptions.Tolerance == EComparisonTolerance::Custom;
-		}
-		else if ( PropertyName == TEXT("ObservationPoint") )
-		{
-			// You can't ever observe from anywhere but the camera on the screenshot test.
-			bIsEditable = false;
-		}
-	}
-
-	return bIsEditable;
-}
-
-void AScreenshotFunctionalTest::PostEditChangeProperty(FPropertyChangedEvent& PropertyChangedEvent)
-{
-	Super::PostEditChangeProperty(PropertyChangedEvent);
-
-	if ( PropertyChangedEvent.Property )
-	{
-		const FName PropertyName = PropertyChangedEvent.Property->GetFName();
-
-		if ( PropertyName == GET_MEMBER_NAME_CHECKED(FAutomationScreenshotOptions, Tolerance) )
-		{
-			ScreenshotOptions.SetToleranceAmounts(ScreenshotOptions.Tolerance);
-		}
-	}
-}
-
-#endif
diff --git a/Engine/Source/Developer/FunctionalTesting/Private/ScreenshotFunctionalTestBase.cpp b/Engine/Source/Developer/FunctionalTesting/Private/ScreenshotFunctionalTestBase.cpp
new file mode 100644
--- /dev/null
+++ b/Engine/Source/Developer/FunctionalTesting/Private/ScreenshotFunctionalTestBase.cpp
@@ -0,0 +1,222 @@
+// Copyright 1998-2017 Epic Games, Inc. All Rights Reserved.
+
+#include "ScreenshotFunctionalTestBase.h"
+
+#include "Engine/GameViewportClient.h"
+#include "AutomationBlueprintFunctionLibrary.h"
+#include "Camera/CameraComponent.h"
+#include "Camera/PlayerCameraManager.h"
+#include "Kismet/GameplayStatics.h"
+#include "Engine/Engine.h"
+#include "EngineGlobals.h"
+#include "Misc/AutomationTest.h"
+#include "Slate/SceneViewport.h"
+#include "RenderingThread.h"
+#include "Tests/AutomationCommon.h"
+#include "Logging/LogMacros.h"
+
+DEFINE_LOG_CATEGORY_STATIC(LogScreenshotFunctionalTest, Log, Log)
+
+AScreenshotFunctionalTestBase::AScreenshotFunctionalTestBase(const FObjectInitializer& ObjectInitializer)
+	: AFunctionalTest(ObjectInitializer)
+	, ScreenshotOptions(EComparisonTolerance::Low)
+{
+	ScreenshotCamera = CreateDefaultSubobject<UCameraComponent>(TEXT("Camera"));
+	ScreenshotCamera->SetupAttachment(RootComponent);
+
+#if (WITH_DEV_AUTOMATION_TESTS || WITH_PERF_AUTOMATION_TESTS)
+	ScreenshotEnvSetup = MakeShareable(new FAutomationTestScreenshotEnvSetup());
+#endif
+}
+
+void AScreenshotFunctionalTestBase::PrepareTest()
+{
+	Super::PrepareTest();
+
+	APlayerController* PlayerController = UGameplayStatics::GetPlayerController(GetWorld(), 0);
+	if (PlayerController)
+	{
+		PlayerController->SetViewTarget(this, FViewTargetTransitionParams());
+	}
+
+	PrepareForScreenshot();
+}
+
+bool AScreenshotFunctionalTestBase::IsReady_Implementation()
+{
+	if ((GetWorld()->GetTimeSeconds() - RunTime) > ScreenshotOptions.Delay)
+	{
+		return (GFrameNumber - RunFrame) > 5;
+	}
+
+	return false;
+}
+
+void AScreenshotFunctionalTestBase::StartTest()
+{
+	Super::StartTest();
+
+	FAutomationTestFramework::Get().OnScreenshotTakenAndCompared.AddUObject(this, &AScreenshotFunctionalTestBase::OnScreenshotTakenAndCompared);
+
+	RequestScreenshot();
+}
+
+void AScreenshotFunctionalTestBase::OnScreenshotTakenAndCompared()
+{
+	RestoreGameViewport();
+
+	FAutomationTestFramework::Get().OnScreenshotTakenAndCompared.RemoveAll(this);
+
+	FinishTest(EFunctionalTestResult::Succeeded, TEXT(""));
+}
+
+void AScreenshotFunctionalTestBase::PrepareForScreenshot()
+{
+	check(GEngine->GameViewport && GEngine->GameViewport->GetGameViewport());
+	check(IsInGameThread());
+
+	FlushRenderingCommands();
+
+	// PS4 requires fixed with/height for video output back buffers so
+	// we cannot adjust the size of viewports
+	if (!FPlatformProperties::HasFixedResolution())
+	{
+		FSceneViewport* GameViewport = GEngine->GameViewport->GetGameViewport();
+		ViewportRestoreSize = GameViewport->GetSize();
+		FIntPoint ScreenshotViewportSize = UAutomationBlueprintFunctionLibrary::GetAutomationScreenshotSize(ScreenshotOptions);
+		GameViewport->SetViewportSize(ScreenshotViewportSize.X, ScreenshotViewportSize.Y);
+	}
+
+#if (WITH_DEV_AUTOMATION_TESTS || WITH_PERF_AUTOMATION_TESTS)
+	ScreenshotEnvSetup->Setup(ScreenshotOptions);
+#endif
+}
+
+void AScreenshotFunctionalTestBase::OnScreenShotCaptured(int32 InSizeX, int32 InSizeY, const TArray<FColor>& InImageData)
+{
+	check(GEngine->GameViewport);
+
+	GEngine->GameViewport->OnScreenshotCaptured().RemoveAll(this);
+
+	FAutomationScreenshotData Data = AutomationCommon::BuildScreenshotData(GetWorld()->GetName(), GetName(), InSizeX, InSizeY);
+
+	// Copy the relevant data into the metadata for the screenshot.
+	Data.bHasComparisonRules = true;
+	Data.ToleranceRed = ScreenshotOptions.ToleranceAmount.Red;
+	Data.ToleranceGreen = ScreenshotOptions.ToleranceAmount.Green;
+	Data.ToleranceBlue = ScreenshotOptions.ToleranceAmount.Blue;
+	Data.ToleranceAlpha = ScreenshotOptions.ToleranceAmount.Alpha;
+	Data.ToleranceMinBrightness = ScreenshotOptions.ToleranceAmount.MinBrightness;
+	Data.ToleranceMaxBrightness = ScreenshotOptions.ToleranceAmount.MaxBrightness;
+	Data.bIgnoreAntiAliasing = ScreenshotOptions.bIgnoreAntiAliasing;
+	Data.bIgnoreColors = ScreenshotOptions.bIgnoreColors;
+	Data.MaximumLocalError = ScreenshotOptions.MaximumLocalError;
+	Data.MaximumGlobalError = ScreenshotOptions.MaximumGlobalError;
+
+	if (GIsAutomationTesting)
+	{
+		FAutomationTestFramework::Get().OnScreenshotCompared.AddUObject(this, &AScreenshotFunctionalTestBase::OnComparisonComplete);
+	}
+
+	FAutomationTestFramework::Get().OnScreenshotCaptured().ExecuteIfBound(InImageData, Data);
+
+	UE_LOG(LogScreenshotFunctionalTest, Log, TEXT("Screenshot captured as %s"), *Data.Path);
+}
+
+void AScreenshotFunctionalTestBase::RequestScreenshot()
+{
+	check(IsInGameThread());
+	check(GEngine->GameViewport);
+
+	// Make sure any screenshot request has been processed
+	FlushRenderingCommands();
+
+	UGameViewportClient* GameViewportClient = GEngine->GameViewport;
+	GameViewportClient->OnScreenshotCaptured().AddUObject(this, &AScreenshotFunctionalTestBase::OnScreenShotCaptured);
+}
+
+void AScreenshotFunctionalTestBase::OnComparisonComplete(bool bWasNew, bool bWasSimilar, double MaxLocalDifference, double GlobalDifference, FString ErrorMessage)
+{
+	FAutomationTestFramework::Get().OnScreenshotCompared.RemoveAll(this);
+
+	if (bWasNew)
+	{
+		UE_LOG(LogScreenshotFunctionalTest, Warning, TEXT("New Screenshot '%s' was discovered!  Please add a ground truth version of it."), *GetName());
+	}
+	else
+	{
+		if (bWasSimilar)
+		{
+			UE_LOG(LogScreenshotFunctionalTest, Display, TEXT("Screenshot '%s' was similar!  Global Difference = %f, Max Local Difference = %f"), *GetName(), GlobalDifference, MaxLocalDifference);
+		}
+		else
+		{
+			if (ErrorMessage.IsEmpty())
+			{
+				UE_LOG(LogScreenshotFunctionalTest, Error, TEXT("Screenshot '%s' test failed, Screnshots were different!  Global Difference = %f, Max Local Difference = %f"), *GetName(), GlobalDifference, MaxLocalDifference);
+			}
+			else
+			{
+				UE_LOG(LogScreenshotFunctionalTest, Error, TEXT("Screenshot '%s' test failed;  Error = %s"), *GetName(), *ErrorMessage);
+			}
+		}
+	}
+
+	FAutomationTestFramework::Get().NotifyScreenshotTakenAndCompared();
+}
+
+void AScreenshotFunctionalTestBase::RestoreGameViewport()
+{
+	check(GEngine->GameViewport && GEngine->GameViewport->GetGameViewport());
+	check(IsInGameThread());
+
+	if (!FPlatformProperties::HasFixedResolution())
+	{
+		FSceneViewport* GameViewport = GEngine->GameViewport->GetGameViewport();
+		GameViewport->SetViewportSize(ViewportRestoreSize.X, ViewportRestoreSize.Y);
+	}
+
+#if (WITH_DEV_AUTOMATION_TESTS || WITH_PERF_AUTOMATION_TESTS)
+	ScreenshotEnvSetup->Restore();
+#endif
+}
+
+#if WITH_EDITOR
+
+bool AScreenshotFunctionalTestBase::CanEditChange(const UProperty* InProperty) const
+{
+	bool bIsEditable = Super::CanEditChange(InProperty);
+	if (bIsEditable && InProperty)
+	{
+		const FName PropertyName = InProperty->GetFName();
+
+		if (PropertyName == GET_MEMBER_NAME_CHECKED(FAutomationScreenshotOptions, ToleranceAmount))
+		{
+			bIsEditable = ScreenshotOptions.Tolerance == EComparisonTolerance::Custom;
+		}
+		else if (PropertyName == TEXT("ObservationPoint"))
+		{
+			// You can't ever observe from anywhere but the camera on the screenshot test.
+			bIsEditable = false;
+		}
+	}
+
+	return bIsEditable;
+}
+
+void AScreenshotFunctionalTestBase::PostEditChangeProperty(FPropertyChangedEvent& PropertyChangedEvent)
+{
+	Super::PostEditChangeProperty(PropertyChangedEvent);
+
+	if (PropertyChangedEvent.Property)
+	{
+		const FName PropertyName = PropertyChangedEvent.Property->GetFName();
+
+		if (PropertyName == GET_MEMBER_NAME_CHECKED(FAutomationScreenshotOptions, Tolerance))
+		{
+			ScreenshotOptions.SetToleranceAmounts(ScreenshotOptions.Tolerance);
+		}
+	}
+}
+
+#endif
diff --git a/Engine/Source/Developer/FunctionalTesting/Public/AutomationBlueprintFunctionLibrary.h b/Engine/Source/Developer/FunctionalTesting/Public/AutomationBlueprintFunctionLibrary.h
--- a/Engine/Source/Developer/FunctionalTesting/Public/AutomationBlueprintFunctionLibrary.h
+++ b/Engine/Source/Developer/FunctionalTesting/Public/AutomationBlueprintFunctionLibrary.h
@@ -24,6 +24,8 @@ public:
 
 	static bool TakeAutomationScreenshotInternal(UObject* WorldContextObject, const FString& Name, FAutomationScreenshotOptions Options);
 
+	static FIntPoint GetAutomationScreenshotSize(const FAutomationScreenshotOptions& Options);
+
 	/**
 	 * Takes a screenshot of the game's viewport.  Does not capture any UI.
 	 */
@@ -83,3 +85,47 @@ public:
 	UFUNCTION(BlueprintPure, Category="Automation")
 	static FAutomationScreenshotOptions GetDefaultScreenshotOptionsForRendering(EComparisonTolerance Tolerance = EComparisonTolerance::Low, float Delay = 0.2);
 };
+
+#if (WITH_DEV_AUTOMATION_TESTS || WITH_PERF_AUTOMATION_TESTS)
+
+template<typename T>
+class FConsoleVariableSwapperTempl
+{
+public:
+	FConsoleVariableSwapperTempl(FString InConsoleVariableName);
+
+	void Set(T Value);
+
+	void Restore();
+
+private:
+	bool bModified;
+	FString ConsoleVariableName;
+
+	T OriginalValue;
+};
+
+class FAutomationTestScreenshotEnvSetup
+{
+public:
+	FAutomationTestScreenshotEnvSetup();
+
+	// Disable AA, auto-exposure, motion blur, contact shadow if InOutOptions.bDisableNoisyRenderingFeatures.
+	// Update screenshot comparison tolerance stored in InOutOptions.
+	// Set visualization buffer name if required.
+	void Setup(FAutomationScreenshotOptions& InOutOptions);
+
+	void Restore();
+
+private:
+	FConsoleVariableSwapperTempl<int32> DefaultFeature_AntiAliasing;
+	FConsoleVariableSwapperTempl<int32> DefaultFeature_AutoExposure;
+	FConsoleVariableSwapperTempl<int32> DefaultFeature_MotionBlur;
+	FConsoleVariableSwapperTempl<int32> PostProcessAAQuality;
+	FConsoleVariableSwapperTempl<int32> MotionBlurQuality;
+	FConsoleVariableSwapperTempl<int32> ScreenSpaceReflectionQuality;
+	FConsoleVariableSwapperTempl<int32> EyeAdaptationQuality;
+	FConsoleVariableSwapperTempl<int32> ContactShadows;
+};
+
+#endif
diff --git a/Engine/Source/Programs/AutomationTool/Scripts/CheckCsprojDotNetVersion.cs b/Engine/Source/Programs/AutomationTool/Scripts/CheckCsprojDotNetVersion.cs
--- a/Engine/Source/Programs/AutomationTool/Scripts/CheckCsprojDotNetVersion.cs
+++ b/Engine/Source/Programs/AutomationTool/Scripts/CheckCsprojDotNetVersion.cs
@@ -26,20 +26,32 @@ class CheckCsprojDotNetVersion : BuildCommand
 		Regex FrameworkRegex = new Regex("<TargetFrameworkVersion>v(\\d\\.\\d\\.?\\d?)<\\/TargetFrameworkVersion>");
 		Regex PossibleAppConfigRegex = new Regex("<TargetFrameworkProfile>(.+)<\\/TargetFrameworkProfile>");
 		Regex AppConfigRegex = new Regex("<supportedRuntime version=\"v(\\d\\.\\d\\.?\\d?)\" sku=\"\\.NETFramework,Version=v(\\d\\.\\d\\.?\\d?),Profile=(.+)\"\\/>");
-		foreach (FileReference CsProj in DirectoryReference.EnumerateFiles(EngineDir, "*.csproj", SearchOption.AllDirectories))
-		{
-			if (CsProj.ContainsName(new FileSystemName("ThirdParty"), EngineDir) ||
-				(CsProj.ContainsName(new FileSystemName("UE4TemplateProject"), EngineDir) && CsProj.GetFileName().Equals("ProjectTemplate.csproj")))
-			{
-				continue;
-			}
+        Regex DotNetCoreRegex = new Regex("<TargetFramework>(netcoreapp2.0|netstandard2.0)<\\/TargetFramework>");
+        foreach (FileReference CsProj in DirectoryReference.EnumerateFiles(EngineDir, "*.csproj", SearchOption.AllDirectories))
+        {
+            if (CsProj.ContainsName(new FileSystemName("ThirdParty"), EngineDir) ||
+                (CsProj.ContainsName(new FileSystemName("UE4TemplateProject"), EngineDir) && CsProj.GetFileName().Equals("ProjectTemplate.csproj")) ||
+                CsProj.GetFileNameWithoutExtension().ToLower().Contains("_mono") ||
+                CsProj.GetFileNameWithoutExtension().ToLower().Contains("unrealvs"))
+
+            {
+                continue;
+            }
+
+            // read in the file
+            string Contents = File.ReadAllText(CsProj.FullName);
+            Match Match = DotNetCoreRegex.Match(Contents);
+            // Check if we're a _NETCore app, ignore these.
+            if (Match.Success)
+            {
+                continue;
+            }
+
 
-			// read in the file
-			string Contents = File.ReadAllText(CsProj.FullName);
-			Match m = FrameworkRegex.Match(Contents);
-			if (m.Success)
+            Match = FrameworkRegex.Match(Contents);
+            if (Match.Success)
 			{
-				string TargetedVersion = m.Groups[1].Value;
+				string TargetedVersion = Match.Groups[1].Value;
 				// make sure we match, throw warning otherwise
 				if (!DesiredTargetVersion.Equals(TargetedVersion, StringComparison.InvariantCultureIgnoreCase))
 				{
@@ -49,8 +61,8 @@ class CheckCsprojDotNetVersion : BuildCommand
 			// if we don't have a TargetFrameworkVersion, check for the existence of TargetFrameworkProfile.
 			else
 			{
-				m = PossibleAppConfigRegex.Match(Contents);
-				if (!m.Success)
+                Match = PossibleAppConfigRegex.Match(Contents);
+				if (!Match.Success)
 				{
 					CommandUtils.Log("No TargetFrameworkVersion or TargetFrameworkProfile found for project {0}, is it a mono project? If not, does it compile properly?", CsProj);
 					continue;
@@ -58,7 +70,7 @@ class CheckCsprojDotNetVersion : BuildCommand
 
 				// look for the app config
 				FileReference AppConfigFile = FileReference.Combine(CsProj.Directory, "app.config");
-				string Profile = m.Groups[1].Value;
+				string Profile = Match.Groups[1].Value;
 				if (!FileReference.Exists(AppConfigFile))
 				{
 					CommandUtils.Log("Found TargetFrameworkProfile but no associated app.config containing the version for project {0}.", CsProj);
@@ -67,8 +79,8 @@ class CheckCsprojDotNetVersion : BuildCommand
 
 				// read in the app config
 				Contents = File.ReadAllText(AppConfigFile.FullName);
-				m = AppConfigRegex.Match(Contents);
-				if (!m.Success)
+                Match = AppConfigRegex.Match(Contents);
+				if (!Match.Success)
 				{
 					CommandUtils.Log("Couldn't find a supportedRuntime match for the version in the app.config for project {0}.", CsProj);
 					continue;
@@ -77,9 +89,9 @@ class CheckCsprojDotNetVersion : BuildCommand
 				// Version1 is the one that appears right after supportedRuntime
 				// Version2 is the one in the sku
 				// ProfileString should match the TargetFrameworkProfile from the csproj
-				string Version1String = m.Groups[1].Value;
-				string Version2String = m.Groups[2].Value;
-				string ProfileString = m.Groups[3].Value;
+				string Version1String = Match.Groups[1].Value;
+				string Version2String = Match.Groups[2].Value;
+				string ProfileString = Match.Groups[3].Value;
 
 				// not sure how this is possible, but check for it anyway
 				if (!ProfileString.Equals(Profile, StringComparison.InvariantCultureIgnoreCase))
diff --git a/Engine/Source/Programs/IncludeTool/IncludeTool/App.config b/Engine/Source/Programs/IncludeTool/IncludeTool/App.config
--- a/Engine/Source/Programs/IncludeTool/IncludeTool/App.config
+++ b/Engine/Source/Programs/IncludeTool/IncludeTool/App.config
@@ -1,6 +1,6 @@
-﻿<?xml version="1.0" encoding="utf-8" ?>
+<?xml version="1.0" encoding="utf-8"?>
 <configuration>
     <startup> 
-        <supportedRuntime version="v4.0" sku=".NETFramework,Version=v4.5.2" />
+        <supportedRuntime version="v4.0" sku=".NETFramework,Version=v4.5"/>
     </startup>
-</configuration>
\ No newline at end of file
+</configuration>
diff --git a/Engine/Source/Programs/IncludeTool/IncludeTool/IncludeTool.csproj b/Engine/Source/Programs/IncludeTool/IncludeTool/IncludeTool.csproj
--- a/Engine/Source/Programs/IncludeTool/IncludeTool/IncludeTool.csproj
+++ b/Engine/Source/Programs/IncludeTool/IncludeTool/IncludeTool.csproj
@@ -9,7 +9,7 @@
     <AppDesignerFolder>Properties</AppDesignerFolder>
     <RootNamespace>IncludeTool</RootNamespace>
     <AssemblyName>IncludeTool</AssemblyName>
-    <TargetFrameworkVersion>v4.5.2</TargetFrameworkVersion>
+    <TargetFrameworkVersion>v4.5</TargetFrameworkVersion>
     <FileAlignment>512</FileAlignment>
     <TargetFrameworkProfile />
   </PropertyGroup>
diff --git a/Engine/Source/Programs/MemoryProfiler2/App.config b/Engine/Source/Programs/MemoryProfiler2/App.config
--- a/Engine/Source/Programs/MemoryProfiler2/App.config
+++ b/Engine/Source/Programs/MemoryProfiler2/App.config
@@ -1,7 +1,7 @@
 <?xml version="1.0"?>
 <configuration>
 	<startup>
-		<supportedRuntime version="v4.0" sku=".NETFramework,Version=v4.0,Profile=Client"/>
+		<supportedRuntime version="v4.0" sku=".NETFramework,Version=v4.5"/>
 	</startup>
 	<runtime>
 		<generatePublisherEvidence enabled="false"/>
diff --git a/Engine/Source/Programs/MemoryProfiler2/MemoryProfiler2.csproj b/Engine/Source/Programs/MemoryProfiler2/MemoryProfiler2.csproj
--- a/Engine/Source/Programs/MemoryProfiler2/MemoryProfiler2.csproj
+++ b/Engine/Source/Programs/MemoryProfiler2/MemoryProfiler2.csproj
@@ -23,9 +23,10 @@
     <OldToolsVersion>3.5</OldToolsVersion>
     <UpgradeBackupLocation>
     </UpgradeBackupLocation>
-    <TargetFrameworkVersion>v4.0</TargetFrameworkVersion>
+    <TargetFrameworkVersion>v4.5</TargetFrameworkVersion>
     <IsWebBootstrapper>false</IsWebBootstrapper>
-    <TargetFrameworkProfile>Client</TargetFrameworkProfile>
+    <TargetFrameworkProfile>
+    </TargetFrameworkProfile>
     <IsWebBootstrapper>false</IsWebBootstrapper>
     <PublishUrl>publish\</PublishUrl>
     <Install>true</Install>
@@ -51,6 +52,7 @@
     <ErrorReport>prompt</ErrorReport>
     <WarningLevel>4</WarningLevel>
     <PlatformTarget>AnyCPU</PlatformTarget>
+    <Prefer32Bit>false</Prefer32Bit>
   </PropertyGroup>
   <PropertyGroup Condition=" '$(Configuration)|$(Platform)' == 'Release|AnyCPU' ">
     <DebugType>pdbonly</DebugType>
@@ -60,6 +62,7 @@
     <ErrorReport>prompt</ErrorReport>
     <WarningLevel>4</WarningLevel>
     <PlatformTarget>AnyCPU</PlatformTarget>
+    <Prefer32Bit>false</Prefer32Bit>
   </PropertyGroup>
   <PropertyGroup>
     <ApplicationIcon>Icon.ico</ApplicationIcon>
diff --git a/Engine/Source/Runtime/D3D12RHI/Public/D3D12Viewport.h b/Engine/Source/Runtime/D3D12RHI/Public/D3D12Viewport.h
--- a/Engine/Source/Runtime/D3D12RHI/Public/D3D12Viewport.h
+++ b/Engine/Source/Runtime/D3D12RHI/Public/D3D12Viewport.h
@@ -107,7 +107,7 @@ public:
 	void IssueFrameEvent();
 
 	IDXGISwapChain1* GetSwapChain() const { return SwapChain1; }
-	IDXGISwapChain1* GetSDRSwapChain() const { return (PixelFormat == SDRPixelFormat) ? GetSwapChain() : SDRSwapChain1; }
+	IDXGISwapChain1* GetSDRSwapChain() const { return (PixelFormat == SDRPixelFormat) ? GetSwapChain() : (IDXGISwapChain1*)SDRSwapChain1; }
 
 	virtual void* GetNativeSwapChain() const override { return GetSwapChain(); }
 	virtual void* GetNativeBackBufferTexture() const override { return GetBackBuffer_RHIThread()->GetResource(); }
diff --git a/Engine/Source/Runtime/Engine/Private/Texture2D.cpp b/Engine/Source/Runtime/Engine/Private/Texture2D.cpp
--- a/Engine/Source/Runtime/Engine/Private/Texture2D.cpp
+++ b/Engine/Source/Runtime/Engine/Private/Texture2D.cpp
@@ -1116,6 +1116,7 @@ FTexture2DResource::FTexture2DResource( UTexture2D* InOwner, int32 InitialMipCou
 :	Owner( InOwner )
 ,	ResourceMem( InOwner->ResourceMem )
 ,	bReadyForStreaming(false)
+,	bUseVirtualUpdatePath(false)
 #if STATS
 ,	TextureSize( 0 )
 #endif
@@ -1185,7 +1186,7 @@ void FTexture2DResource::InitRHI()
 	uint32 SizeY = OwnerMips[CurrentFirstMip].SizeY;
 
 	// Create the RHI texture.
-	uint32 TexCreateFlags = (Owner->SRGB ? TexCreate_SRGB : 0) | TexCreate_OfflineProcessed;
+	uint32 TexCreateFlags = (Owner->SRGB ? TexCreate_SRGB : 0) | TexCreate_OfflineProcessed | TexCreate_Streamable;
 	// if no miptail is available then create the texture without a packed miptail
 	if( Owner->GetMipTailBaseIndex() == -1 )
 	{
@@ -1197,6 +1198,10 @@ void FTexture2DResource::InitRHI()
 		TexCreateFlags |= TexCreate_NoTiling;
 	}
 
+	// Determine if this texture should use the virtual update path when streaming in and out mips. 
+	// Note that because of "r.VirtualTextureReducedMemory" it might use a virtual allocation initially.
+	bUseVirtualUpdatePath = CanCreateAsVirtualTexture(TexCreateFlags);
+
 	EPixelFormat EffectiveFormat = Owner->GetPixelFormat();
 
 	CreateSamplerStates(UTexture2D::GetGlobalMipMapLODBias() + GetDefaultMipMapBias());
@@ -1215,7 +1220,7 @@ void FTexture2DResource::InitRHI()
 			static auto CVarVirtualTextureReducedMemoryEnabled = IConsoleManager::Get().FindTConsoleVariableDataInt(TEXT("r.VirtualTextureReducedMemory"));
 			check(CVarVirtualTextureReducedMemoryEnabled);
 
-			if ( Owner->bIsStreamable && CanCreateAsVirtualTexture(TexCreateFlags) && (CVarVirtualTextureReducedMemoryEnabled->GetValueOnRenderThread() == 0 || RequestedMips > UTexture2D::GetMinTextureResidentMipCount()) )
+			if ( Owner->bIsStreamable && bUseVirtualUpdatePath && (CVarVirtualTextureReducedMemoryEnabled->GetValueOnRenderThread() == 0 || RequestedMips > UTexture2D::GetMinTextureResidentMipCount()) )
 			{
 				TexCreateFlags |= TexCreate_Virtual;
 
@@ -1468,7 +1473,7 @@ bool UTexture2D::StreamIn(int32 NewMipCount, bool bHighPrio)
 {
 	check(IsInGameThread());
 	FTexture2DResource* Texture2DResource = (FTexture2DResource*)Resource;
-	if (bIsStreamable && !PendingUpdate && Texture2DResource && Texture2DResource->bReadyForStreaming && Texture2DResource->TextureRHI && NewMipCount > GetNumResidentMips())
+	if (bIsStreamable && !PendingUpdate && Texture2DResource && Texture2DResource->bReadyForStreaming && NewMipCount > GetNumResidentMips())
 	{
 #if WITH_EDITORONLY_DATA
 		if (FPlatformProperties::HasEditorOnlyData())
@@ -1485,11 +1490,8 @@ bool UTexture2D::StreamIn(int32 NewMipCount, bool bHighPrio)
 		else
 #endif
 		{
-			const uint32 TexCreateFlags = Texture2DResource->TextureRHI->GetFlags();
-			const bool bIsVirtualTexture = (TexCreateFlags & TexCreate_Virtual) == TexCreate_Virtual;
-
 			// If the future texture is to be a virtual texture, use the virtual stream in path.
-			if (CanCreateAsVirtualTexture(TexCreateFlags) || bIsVirtualTexture)
+			if (Texture2DResource->bUseVirtualUpdatePath)
 			{
 				PendingUpdate = new FTexture2DStreamIn_IO_Virtual(this, NewMipCount, bHighPrio);
 			}
@@ -1516,12 +1518,9 @@ bool UTexture2D::StreamOut(int32 NewMipCount)
 {
 	check(IsInGameThread());
 	FTexture2DResource* Texture2DResource = (FTexture2DResource*)Resource;
-	if (bIsStreamable && !PendingUpdate && Texture2DResource && Texture2DResource->bReadyForStreaming && Texture2DResource->TextureRHI && NewMipCount < GetNumResidentMips())
+	if (bIsStreamable && !PendingUpdate && Texture2DResource && Texture2DResource->bReadyForStreaming && NewMipCount < GetNumResidentMips())
 	{
-		const uint32 TexCreateFlags = Texture2DResource->TextureRHI->GetFlags();
-		const bool bIsVirtualTexture = (TexCreateFlags & TexCreate_Virtual) == TexCreate_Virtual;
-
-		if (bIsVirtualTexture)
+		if (Texture2DResource->bUseVirtualUpdatePath)
 		{
 			PendingUpdate = new FTexture2DStreamOut_Virtual(this, NewMipCount);
 		}
diff --git a/Engine/Source/Runtime/Engine/Private/TextureDerivedData.cpp b/Engine/Source/Runtime/Engine/Private/TextureDerivedData.cpp
--- a/Engine/Source/Runtime/Engine/Private/TextureDerivedData.cpp
+++ b/Engine/Source/Runtime/Engine/Private/TextureDerivedData.cpp
@@ -817,7 +817,7 @@ bool FTexturePlatformData::TryLoadMips(int32 FirstMipToLoad, void** OutMipData)
 				OutMipData[MipIndex - FirstMipToLoad] = FMemory::Malloc(Mip.BulkData.GetBulkDataSize());
 #if 0
 				checkSlow(!Mip.BulkData.GetFilename().EndsWith(TEXT(".ubulk"))); // We want to make sure that any non-streamed mips are coming from the texture asset file, and not from an external bulk file
-#else
+#elif !(WITH_DEV_AUTOMATION_TESTS || WITH_PERF_AUTOMATION_TESTS)
 				UE_CLOG(Mip.BulkData.GetFilename().EndsWith(TEXT(".ubulk")), LogTexture, Error, TEXT("Loading non-streamed mips from an external bulk file.  This is not desireable.  File %s"), *(Mip.BulkData.GetFilename() ) );
 #endif
 				
diff --git a/Engine/Source/Runtime/Engine/Public/TextureResource.h b/Engine/Source/Runtime/Engine/Public/TextureResource.h
--- a/Engine/Source/Runtime/Engine/Public/TextureResource.h
+++ b/Engine/Source/Runtime/Engine/Public/TextureResource.h
@@ -148,6 +148,9 @@ private:
 	/** Whether the texture RHI has been initialized.														*/
 	bool bReadyForStreaming;
 
+	/** Whether this texture should be updated using the virtual allocations.								*/
+	bool bUseVirtualUpdatePath;
+
 	EMipFadeSettings MipFadeSetting;
 
 	/** First mip level used in Texture2DRHI. This is always correct as long as Texture2DRHI is allocated, regardless of streaming status. */
diff --git a/Engine/Source/Runtime/HeadMountedDisplay/Private/HeadMountedDisplayConsoleCommands.cpp b/Engine/Source/Runtime/HeadMountedDisplay/Private/HeadMountedDisplayConsoleCommands.cpp
--- a/Engine/Source/Runtime/HeadMountedDisplay/Private/HeadMountedDisplayConsoleCommands.cpp
+++ b/Engine/Source/Runtime/HeadMountedDisplay/Private/HeadMountedDisplayConsoleCommands.cpp
@@ -25,7 +25,7 @@ DEFINE_LOG_CATEGORY_STATIC(LogHeadMountedDisplayCommand, Display, All);
 */
 static TAutoConsoleVariable<int32> CVarHiddenAreaMask(
 	TEXT("vr.HiddenAreaMask"),
-	0,
+	1,
 	*LOCTEXT("CVarText_HiddenAreaMask", "Enable or disable hidden area mask\n0: disabled\n1: enabled").ToString(),
 	ECVF_Scalability | ECVF_RenderThreadSafe);
 
diff --git a/Engine/Source/Runtime/Launch/Resources/Version.h b/Engine/Source/Runtime/Launch/Resources/Version.h
--- a/Engine/Source/Runtime/Launch/Resources/Version.h
+++ b/Engine/Source/Runtime/Launch/Resources/Version.h
@@ -55,7 +55,7 @@
 // These numbers define the banner UE4 version, and are the most significant numbers when ordering two engine versions (that is, a 4.12.* version is always 
 // newer than a 4.11.* version, regardless of the changelist that it was built with)
 #define ENGINE_MAJOR_VERSION	4
-#define ENGINE_MINOR_VERSION	18
+#define ENGINE_MINOR_VERSION	19
 #define ENGINE_PATCH_VERSION	0
 
 // If set to 1, indicates that this is a licensee build of the engine. For the same major/minor/patch release of the engine, licensee changelists are always 
diff --git a/Engine/Source/Runtime/OpenGLDrv/Private/Android/AndroidESDeferredOpenGL.h b/Engine/Source/Runtime/OpenGLDrv/Private/Android/AndroidESDeferredOpenGL.h
--- a/Engine/Source/Runtime/OpenGLDrv/Private/Android/AndroidESDeferredOpenGL.h
+++ b/Engine/Source/Runtime/OpenGLDrv/Private/Android/AndroidESDeferredOpenGL.h
@@ -46,6 +46,7 @@ typedef FEGLGLSync UGLsync;
 #define GLdouble		GLfloat
 #define GL_BGRA			GL_BGRA_EXT
 #define GL_UNSIGNED_INT_8_8_8_8_REV	GL_UNSIGNED_BYTE
+#define GL_UNSIGNED_INT_8_8_8_8	0x8035
 
 #ifndef EGL_KHR_create_context
 #define EGL_KHR_create_context 1
diff --git a/Engine/Source/Runtime/RHI/Public/RHIDefinitions.h b/Engine/Source/Runtime/RHI/Public/RHIDefinitions.h
--- a/Engine/Source/Runtime/RHI/Public/RHIDefinitions.h
+++ b/Engine/Source/Runtime/RHI/Public/RHIDefinitions.h
@@ -640,6 +640,8 @@ enum ETextureCreateFlags
 	TexCreate_NoFastClear = 1 << 25,
 	// Texture is a depth stencil resolve target
 	TexCreate_DepthStencilResolveTarget = 1 << 26,
+	// Flag used to indicted this texture is a streamable 2D texture, and should be counted towards the texture streaming pool budget.
+	TexCreate_Streamable = 1 << 27,
 	// Render target will not FinalizeFastClear; Caches and meta data will be flushed, but clearing will be skipped (avoids potentially trashing metadata)
 	TexCreate_NoFastClearFinalize = 1 << 28,
 	// Hint to the driver that this resource is managed properly by the engine for Alternate-Frame-Rendering in mGPU usage.
