Static analysis fix.

diff --git a/Engine/Source/Developer/FunctionalTesting/Private/AutomationBlueprintFunctionLibrary.cpp b/Engine/Source/Developer/FunctionalTesting/Private/AutomationBlueprintFunctionLibrary.cpp
--- a/Engine/Source/Developer/FunctionalTesting/Private/AutomationBlueprintFunctionLibrary.cpp
+++ b/Engine/Source/Developer/FunctionalTesting/Private/AutomationBlueprintFunctionLibrary.cpp
@@ -374,7 +374,7 @@ bool UAutomationBlueprintFunctionLibrary::TakeAutomationScreenshotInternal(UObje
 #endif
 
 	FScreenshotRequest::RequestScreenshot(false);
-	return true;
+	return true; //-V773
 }
 
 void UAutomationBlueprintFunctionLibrary::TakeAutomationScreenshot(UObject* WorldContextObject, FLatentActionInfo LatentInfo, const FString& Name, const FAutomationScreenshotOptions& Options)
