Merging fix for UE-6025,  UE-60260, UE-60280 from //UE4/Release-Staging-4.20

Added GroundTruth for UE-60162

Merged Testblacklist work from Release-4.20

#tests ran locally

diff --git a/Engine/Source/Developer/FunctionalTesting/Private/FunctionalTestingModule.cpp b/Engine/Source/Developer/FunctionalTesting/Private/FunctionalTestingModule.cpp
--- a/Engine/Source/Developer/FunctionalTesting/Private/FunctionalTestingModule.cpp
+++ b/Engine/Source/Developer/FunctionalTesting/Private/FunctionalTestingModule.cpp
@@ -11,6 +11,7 @@
 #include "EngineGlobals.h"
 #include "ARFilter.h"
 #include "AssetRegistryModule.h"
+#include "Misc/CommandLine.h"
 
 #define LOCTEXT_NAMESPACE "FunctionalTesting"
 
@@ -38,9 +39,16 @@ private:
 	void OnGetAssetTagsForWorld(const UWorld* World, TArray<UObject::FAssetRegistryTag>& OutTags);
 
 	void BuildTestBlacklistFromConfig();
-	bool IsBlacklisted(const FString& MapName, const FString& TestName) const;
+	bool IsBlacklisted(const FString& MapName, const FString& TestName, FString& OutReason) const;
 
-	TArray<FString> TestBlacklist;
+	struct FBlacklistEntry
+	{
+		FString	Map;
+		FString Test;
+		FString Reason;
+	};
+
+	TMap<FString, FBlacklistEntry> TestBlacklist;
 	TWeakObjectPtr<class UFunctionalTestingManager> TestManager;
 	bool bPendingActivation;
 };
@@ -148,12 +156,17 @@ void FFunctionalTestingModule::GetMapTests(bool bEditorOnlyTests, TArray<FString
 
 							if (MapTest.Split(TEXT("|"), &BeautifulTestName, &RealTestName))
 							{
-								if (!IsBlacklisted(MapPackageName, RealTestName))
+								FString BlacklistReason;
+								if (!IsBlacklisted(MapPackageName, RealTestName, BlacklistReason))
 								{
 									OutBeautifiedNames.Add(MapPackageName + TEXT(".") + *BeautifulTestName);
 									OutTestCommands.Add(MapAssetPath + TEXT(";") + MapAsset.PackageName.ToString() + TEXT(";") + *RealTestName);
 									OutTestMapAssets.AddUnique(MapAssetPath);
 								}
+								else
+								{
+									UE_LOG(LogFunctionalTest, Warning, TEXT("Test '%s' is blacklisted. %s"), *MapTest, *BlacklistReason);
+								}
 							}
 						}
 					}
@@ -266,30 +279,47 @@ void FFunctionalTestingModule::BuildTestBlacklistFromConfig()
 	TestBlacklist.Empty();
 	if (GConfig)
 	{
+	
+		const FString CommandLine = FCommandLine::Get();
+
 		for (const TPair<FString,FConfigFile>& Config : *GConfig)
 		{
 			FConfigSection* BlacklistSection = GConfig->GetSectionPrivate(TEXT("AutomationTestBlacklist"), false, true, Config.Key);
 			if (BlacklistSection)
 			{
-				// Parse all blacklist definitions of the format "BlacklistTest=(Map=/Game/Tests/MapName, Test=TestName)"
+				// Parse all blacklist definitions of the format "BlacklistTest=(Map=/Game/Tests/MapName, Test=TestName, Reason="Foo")"
 				for (FConfigSection::TIterator Section(*BlacklistSection); Section; ++Section)
 				{
 					if (Section.Key() == TEXT("BlacklistTest"))
 					{
 						FString BlacklistValue = Section.Value().GetValue();
-						FString Map, Test;
+						FString Map, Test, Reason;
 						bool bSuccess = false;
 						
 						if (FParse::Value(*BlacklistValue, TEXT("Map="), Map, true) && FParse::Value(*BlacklistValue, TEXT("Test="), Test, true))
 						{
+							FParse::Value(*BlacklistValue, TEXT("Reason="), Reason);
+
 							// These are used as folders so ensure they match the expected layout
 							if (Map.StartsWith(TEXT("/")))
 							{
 								FString ListName = Map + TEXT("/") + Test;
 								ListName.RemoveSpacesInline();
-
-								TestBlacklist.AddUnique(ListName);
 								bSuccess = true;
+
+								if (CommandLine.Contains(Map) || CommandLine.Contains(Test))
+								{
+									UE_LOG(LogFunctionalTest, Warning, TEXT("Test '%s' is blacklisted but allowing due to command line."), *BlacklistValue);
+								}
+								else
+								{
+									// convert Pretty.Name into PrettyName as we compare on the latter.
+									ListName = ListName.Replace(TEXT("."), TEXT(""));
+									FBlacklistEntry& Entry = TestBlacklist.Add(ListName);
+									Entry.Map = Map;
+									Entry.Test = Test;
+									Entry.Reason = Reason;
+								}
 							}
 						}
 						
@@ -306,18 +336,26 @@ void FFunctionalTestingModule::BuildTestBlacklistFromConfig()
 	if (TestBlacklist.Num() > 0)
 	{
 		UE_LOG(LogFunctionalTest, Log, TEXT("Automated Test Blacklist:"));
-		for (FString& Test : TestBlacklist)
+		for (auto& KV : TestBlacklist)
 		{
-			UE_LOG(LogFunctionalTest, Log, TEXT("\tTest: %s"), *Test);
+			UE_LOG(LogFunctionalTest, Log, TEXT("\tTest: %s"), *KV.Key);
 		}
 	}
 }
 
-bool FFunctionalTestingModule::IsBlacklisted(const FString& MapName, const FString& TestName) const
+bool FFunctionalTestingModule::IsBlacklisted(const FString& MapName, const FString& TestName, FString& OutReason) const
 {
 	FString ListName = MapName + TEXT("/") + TestName;
 	ListName.RemoveSpacesInline();
-	return TestBlacklist.Contains(ListName);
+
+	const FBlacklistEntry* Entry = TestBlacklist.Find(ListName);
+
+	if (Entry)
+	{
+		OutReason = Entry->Reason;
+	}
+
+	return Entry != nullptr;
 }
 
 //////////////////////////////////////////////////////////////////////////
diff --git a/Engine/Source/Developer/MeshSimplifier/Private/QuadricMeshReduction.cpp b/Engine/Source/Developer/MeshSimplifier/Private/QuadricMeshReduction.cpp
--- a/Engine/Source/Developer/MeshSimplifier/Private/QuadricMeshReduction.cpp
+++ b/Engine/Source/Developer/MeshSimplifier/Private/QuadricMeshReduction.cpp
@@ -552,6 +552,9 @@ public:
 
 			for (const FMeshTriangle& MeshTriangle : Triangles)
 			{
+				int32 CurrentFaceIndex = FaceIndex;
+				//Increment face index here because there is many continue in this for loop
+				++FaceIndex;
 				for (int32 CornerIndex = 0; CornerIndex < 3; ++CornerIndex)
 				{
 					VertexInstanceIDs[CornerIndex] = MeshTriangle.GetVertexInstanceID(CornerIndex);
@@ -570,8 +573,7 @@ public:
 				int32 VertexIndices[3];
 				for (int32 CornerIndex = 0; CornerIndex < 3; CornerIndex++)
 				{
-					int32 WedgeIndex = FaceIndex * 3 + CornerIndex;
-					++FaceIndex;
+					int32 WedgeIndex = CurrentFaceIndex * 3 + CornerIndex;
 
 					TVertSimp< NumTexCoords > NewVert;
 
diff --git a/Engine/Source/Runtime/Engine/Private/StaticMesh.cpp b/Engine/Source/Runtime/Engine/Private/StaticMesh.cpp
--- a/Engine/Source/Runtime/Engine/Private/StaticMesh.cpp
+++ b/Engine/Source/Runtime/Engine/Private/StaticMesh.cpp
@@ -1350,7 +1350,7 @@ FArchive& operator<<(FArchive& Ar, FMeshBuildSettings& BuildSettings)
 // differences, etc.) replace the version GUID below with a new one.
 // In case of merge conflicts with DDC versions, you MUST generate a new GUID
 // and set this new GUID as the version.                                       
-#define STATICMESH_DERIVEDDATA_VER TEXT("42B9CFE0CAFE46F29F43E9049AA21282")
+#define STATICMESH_DERIVEDDATA_VER TEXT("4FDA6C9D166447EAB2DA9EDA8A184591")
 
 static const FString& GetStaticMeshDerivedDataVersion()
 {
