Merging from //UE4/Release-4.20 - Added ability to suppress warning for blacklisted automation tests

diff --git a/Engine/Source/Developer/FunctionalTesting/Private/FunctionalTestingModule.cpp b/Engine/Source/Developer/FunctionalTesting/Private/FunctionalTestingModule.cpp
--- a/Engine/Source/Developer/FunctionalTesting/Private/FunctionalTestingModule.cpp
+++ b/Engine/Source/Developer/FunctionalTesting/Private/FunctionalTestingModule.cpp
@@ -39,13 +39,17 @@ private:
 	void OnGetAssetTagsForWorld(const UWorld* World, TArray<UObject::FAssetRegistryTag>& OutTags);
 
 	void BuildTestBlacklistFromConfig();
-	bool IsBlacklisted(const FString& MapName, const FString& TestName, FString& OutReason) const;
+	bool IsBlacklisted(const FString& MapName, const FString& TestName, FString* OutReason=nullptr, bool* OutWarn=nullptr) const;
 
 	struct FBlacklistEntry
 	{
+		FBlacklistEntry() :
+			bWarn(false) {}
+
 		FString	Map;
 		FString Test;
 		FString Reason;
+		bool bWarn;
 	};
 
 	TMap<FString, FBlacklistEntry> TestBlacklist;
@@ -157,7 +161,8 @@ void FFunctionalTestingModule::GetMapTests(bool bEditorOnlyTests, TArray<FString
 							if (MapTest.Split(TEXT("|"), &BeautifulTestName, &RealTestName))
 							{
 								FString BlacklistReason;
-								if (!IsBlacklisted(MapPackageName, RealTestName, BlacklistReason))
+								bool bWarn(false);
+								if (!IsBlacklisted(MapPackageName, RealTestName, &BlacklistReason, &bWarn))
 								{
 									OutBeautifiedNames.Add(MapPackageName + TEXT(".") + *BeautifulTestName);
 									OutTestCommands.Add(MapAssetPath + TEXT(";") + MapAsset.PackageName.ToString() + TEXT(";") + *RealTestName);
@@ -165,10 +170,14 @@ void FFunctionalTestingModule::GetMapTests(bool bEditorOnlyTests, TArray<FString
 								}
 								else
 								{
-									if (BlacklistReason.Len())
+									if (bWarn)
 									{
 										UE_LOG(LogFunctionalTest, Warning, TEXT("Test '%s' is blacklisted. %s"), *MapTest, *BlacklistReason);
 									}
+									else
+									{
+										UE_LOG(LogFunctionalTest, Display, TEXT("Test '%s' is blacklisted. %s"), *MapTest, *BlacklistReason);
+									}
 								}
 							}
 						}
@@ -296,12 +305,13 @@ void FFunctionalTestingModule::BuildTestBlacklistFromConfig()
 					if (Section.Key() == TEXT("BlacklistTest"))
 					{
 						FString BlacklistValue = Section.Value().GetValue();
-						FString Map, Test, Reason;
+						FString Map, Test, Reason, Warn;
 						bool bSuccess = false;
 						
 						if (FParse::Value(*BlacklistValue, TEXT("Map="), Map, true) && FParse::Value(*BlacklistValue, TEXT("Test="), Test, true))
 						{
 							FParse::Value(*BlacklistValue, TEXT("Reason="), Reason);
+							FParse::Value(*BlacklistValue, TEXT("Warn="), Warn);
 
 							// These are used as folders so ensure they match the expected layout
 							if (Map.StartsWith(TEXT("/")))
@@ -322,6 +332,7 @@ void FFunctionalTestingModule::BuildTestBlacklistFromConfig()
 									Entry.Map = Map;
 									Entry.Test = Test;
 									Entry.Reason = Reason;
+									Entry.bWarn = Warn.ToBool();
 								}
 							}
 						}
@@ -346,7 +357,7 @@ void FFunctionalTestingModule::BuildTestBlacklistFromConfig()
 	}
 }
 
-bool FFunctionalTestingModule::IsBlacklisted(const FString& MapName, const FString& TestName, FString& OutReason) const
+bool FFunctionalTestingModule::IsBlacklisted(const FString& MapName, const FString& TestName, FString* OutReason, bool *OutWarn) const
 {
 	FString ListName = MapName + TEXT("/") + TestName;
 	ListName.RemoveSpacesInline();
@@ -355,7 +366,15 @@ bool FFunctionalTestingModule::IsBlacklisted(const FString& MapName, const FStri
 
 	if (Entry)
 	{
-		OutReason = Entry->Reason;
+		if (OutReason != nullptr)
+		{
+			*OutReason = Entry->Reason;
+		}
+
+		if (OutWarn != nullptr)
+		{
+			*OutWarn = Entry->bWarn;
+		}
 	}
 
 	return Entry != nullptr;
