Merging //UE4/Dev-Main to Dev-Framework (//UE4/Dev-Framework)

#rnx

diff --git a/Engine/Plugins/FX/Niagara/Source/NiagaraEditor/Private/Sequencer/NiagaraSequence/Sections/MovieSceneNiagaraEmitterSection.cpp b/Engine/Plugins/FX/Niagara/Source/NiagaraEditor/Private/Sequencer/NiagaraSequence/Sections/MovieSceneNiagaraEmitterSection.cpp
--- a/Engine/Plugins/FX/Niagara/Source/NiagaraEditor/Private/Sequencer/NiagaraSequence/Sections/MovieSceneNiagaraEmitterSection.cpp
+++ b/Engine/Plugins/FX/Niagara/Source/NiagaraEditor/Private/Sequencer/NiagaraSequence/Sections/MovieSceneNiagaraEmitterSection.cpp
@@ -99,8 +99,9 @@ TSharedPtr<FStructOnScope> GetKeyStruct(TMovieSceneChannelHandle<FMovieSceneNiag
 	int32 KeyValueIndex = Channel.Get()->GetData().GetIndex(InHandle);
 	if (KeyValueIndex != INDEX_NONE)
 	{
-		FNiagaraEmitterSectionKey& KeyValue = Channel.Get()->GetData().GetValues()[KeyValueIndex];
-		return MakeShared<FStructOnScope>(KeyValue.Value.GetType().GetStruct(), KeyValue.Value.GetData());
+		FNiagaraTypeDefinition KeyType = Channel.Get()->GetData().GetValues()[KeyValueIndex].Value.GetType();
+		uint8* KeyData = Channel.Get()->GetData().GetValues()[KeyValueIndex].Value.GetData();
+		return MakeShared<FStructOnScope>(KeyType.GetStruct(), KeyData);
 	}
 	return TSharedPtr<FStructOnScope>();
 }
diff --git a/Engine/Shaders/Private/ReflectionEnvironmentPixelShader.usf b/Engine/Shaders/Private/ReflectionEnvironmentPixelShader.usf
--- a/Engine/Shaders/Private/ReflectionEnvironmentPixelShader.usf
+++ b/Engine/Shaders/Private/ReflectionEnvironmentPixelShader.usf
@@ -298,7 +298,7 @@ void ReflectionEnvironmentSkyLighting(
 {
 	float2 ScreenUV = UVAndScreenPos.xy;
 	FScreenSpaceData ScreenSpaceData = GetScreenSpaceData(ScreenUV);
-	uint ShadingModelID = GetShadingModelId(ScreenUV);
+	uint ShadingModelID = ScreenSpaceData.GBuffer.ShadingModelID;
 
 	OutColor = 0.0f;
 
diff --git a/Engine/Source/Developer/FunctionalTesting/Private/FunctionalTestingModule.cpp b/Engine/Source/Developer/FunctionalTesting/Private/FunctionalTestingModule.cpp
--- a/Engine/Source/Developer/FunctionalTesting/Private/FunctionalTestingModule.cpp
+++ b/Engine/Source/Developer/FunctionalTesting/Private/FunctionalTestingModule.cpp
@@ -11,6 +11,7 @@
 #include "EngineGlobals.h"
 #include "ARFilter.h"
 #include "AssetRegistryModule.h"
+#include "Misc/CommandLine.h"
 
 #define LOCTEXT_NAMESPACE "FunctionalTesting"
 
@@ -38,9 +39,16 @@ private:
 	void OnGetAssetTagsForWorld(const UWorld* World, TArray<UObject::FAssetRegistryTag>& OutTags);
 
 	void BuildTestBlacklistFromConfig();
-	bool IsBlacklisted(const FString& MapName, const FString& TestName) const;
+	bool IsBlacklisted(const FString& MapName, const FString& TestName, FString& OutReason) const;
 
-	TArray<FString> TestBlacklist;
+	struct FBlacklistEntry
+	{
+		FString	Map;
+		FString Test;
+		FString Reason;
+	};
+
+	TMap<FString, FBlacklistEntry> TestBlacklist;
 	TWeakObjectPtr<class UFunctionalTestingManager> TestManager;
 	bool bPendingActivation;
 };
@@ -148,12 +156,20 @@ void FFunctionalTestingModule::GetMapTests(bool bEditorOnlyTests, TArray<FString
 
 							if (MapTest.Split(TEXT("|"), &BeautifulTestName, &RealTestName))
 							{
-								if (!IsBlacklisted(MapPackageName, RealTestName))
+								FString BlacklistReason;
+								if (!IsBlacklisted(MapPackageName, RealTestName, BlacklistReason))
 								{
 									OutBeautifiedNames.Add(MapPackageName + TEXT(".") + *BeautifulTestName);
 									OutTestCommands.Add(MapAssetPath + TEXT(";") + MapAsset.PackageName.ToString() + TEXT(";") + *RealTestName);
 									OutTestMapAssets.AddUnique(MapAssetPath);
 								}
+								else
+								{
+									if (BlacklistReason.Len())
+									{
+										UE_LOG(LogFunctionalTest, Warning, TEXT("Test '%s' is blacklisted. %s"), *MapTest, *BlacklistReason);
+									}
+								}
 							}
 						}
 					}
@@ -266,30 +282,47 @@ void FFunctionalTestingModule::BuildTestBlacklistFromConfig()
 	TestBlacklist.Empty();
 	if (GConfig)
 	{
+	
+		const FString CommandLine = FCommandLine::Get();
+
 		for (const TPair<FString,FConfigFile>& Config : *GConfig)
 		{
 			FConfigSection* BlacklistSection = GConfig->GetSectionPrivate(TEXT("AutomationTestBlacklist"), false, true, Config.Key);
 			if (BlacklistSection)
 			{
-				// Parse all blacklist definitions of the format "BlacklistTest=(Map=/Game/Tests/MapName, Test=TestName)"
+				// Parse all blacklist definitions of the format "BlacklistTest=(Map=/Game/Tests/MapName, Test=TestName, Reason="Foo")"
 				for (FConfigSection::TIterator Section(*BlacklistSection); Section; ++Section)
 				{
 					if (Section.Key() == TEXT("BlacklistTest"))
 					{
 						FString BlacklistValue = Section.Value().GetValue();
-						FString Map, Test;
+						FString Map, Test, Reason;
 						bool bSuccess = false;
 						
 						if (FParse::Value(*BlacklistValue, TEXT("Map="), Map, true) && FParse::Value(*BlacklistValue, TEXT("Test="), Test, true))
 						{
+							FParse::Value(*BlacklistValue, TEXT("Reason="), Reason);
+
 							// These are used as folders so ensure they match the expected layout
 							if (Map.StartsWith(TEXT("/")))
 							{
 								FString ListName = Map + TEXT("/") + Test;
 								ListName.RemoveSpacesInline();
-
-								TestBlacklist.AddUnique(ListName);
 								bSuccess = true;
+
+								if (CommandLine.Contains(Map) || CommandLine.Contains(Test))
+								{
+									UE_LOG(LogFunctionalTest, Warning, TEXT("Test '%s' is blacklisted but allowing due to command line."), *BlacklistValue);
+								}
+								else
+								{
+									// convert Pretty.Name into PrettyName as we compare on the latter.
+									ListName = ListName.Replace(TEXT("."), TEXT(""));
+									FBlacklistEntry& Entry = TestBlacklist.Add(ListName);
+									Entry.Map = Map;
+									Entry.Test = Test;
+									Entry.Reason = Reason;
+								}
 							}
 						}
 						
@@ -306,18 +339,26 @@ void FFunctionalTestingModule::BuildTestBlacklistFromConfig()
 	if (TestBlacklist.Num() > 0)
 	{
 		UE_LOG(LogFunctionalTest, Log, TEXT("Automated Test Blacklist:"));
-		for (FString& Test : TestBlacklist)
+		for (auto& KV : TestBlacklist)
 		{
-			UE_LOG(LogFunctionalTest, Log, TEXT("\tTest: %s"), *Test);
+			UE_LOG(LogFunctionalTest, Log, TEXT("\tTest: %s"), *KV.Key);
 		}
 	}
 }
 
-bool FFunctionalTestingModule::IsBlacklisted(const FString& MapName, const FString& TestName) const
+bool FFunctionalTestingModule::IsBlacklisted(const FString& MapName, const FString& TestName, FString& OutReason) const
 {
 	FString ListName = MapName + TEXT("/") + TestName;
 	ListName.RemoveSpacesInline();
-	return TestBlacklist.Contains(ListName);
+
+	const FBlacklistEntry* Entry = TestBlacklist.Find(ListName);
+
+	if (Entry)
+	{
+		OutReason = Entry->Reason;
+	}
+
+	return Entry != nullptr;
 }
 
 //////////////////////////////////////////////////////////////////////////
diff --git a/Engine/Source/Developer/MeshSimplifier/Private/QuadricMeshReduction.cpp b/Engine/Source/Developer/MeshSimplifier/Private/QuadricMeshReduction.cpp
--- a/Engine/Source/Developer/MeshSimplifier/Private/QuadricMeshReduction.cpp
+++ b/Engine/Source/Developer/MeshSimplifier/Private/QuadricMeshReduction.cpp
@@ -552,6 +552,9 @@ public:
 
 			for (const FMeshTriangle& MeshTriangle : Triangles)
 			{
+				int32 CurrentFaceIndex = FaceIndex;
+				//Increment face index here because there is many continue in this for loop
+				++FaceIndex;
 				for (int32 CornerIndex = 0; CornerIndex < 3; ++CornerIndex)
 				{
 					VertexInstanceIDs[CornerIndex] = MeshTriangle.GetVertexInstanceID(CornerIndex);
@@ -570,8 +573,7 @@ public:
 				int32 VertexIndices[3];
 				for (int32 CornerIndex = 0; CornerIndex < 3; CornerIndex++)
 				{
-					int32 WedgeIndex = FaceIndex * 3 + CornerIndex;
-					++FaceIndex;
+					int32 WedgeIndex = CurrentFaceIndex * 3 + CornerIndex;
 
 					TVertSimp< NumTexCoords > NewVert;
 
diff --git a/Engine/Source/Editor/UnrealEd/Private/ObjectTools.cpp b/Engine/Source/Editor/UnrealEd/Private/ObjectTools.cpp
--- a/Engine/Source/Editor/UnrealEd/Private/ObjectTools.cpp
+++ b/Engine/Source/Editor/UnrealEd/Private/ObjectTools.cpp
@@ -785,15 +785,15 @@ namespace ObjectTools
 			}
 		}
 
-		for (int CurrentIndex = 0; CurrentIndex < ReferencingPropertiesMapKeys.Num(); CurrentIndex++)
 		{
-			int CheckForCircles = 0;
-			FFindReferencersArchive FindDependentArchive(ReferencingPropertiesMapKeys[CurrentIndex], ReferencingPropertiesMapKeys);
-			for (int DependentIndex = CurrentIndex + 1; DependentIndex < ReferencingPropertiesMapKeys.Num(); DependentIndex++)
+			TBitArray<> TouchedThisItteration(false, ReferencingPropertiesMapKeys.Num());
+			for (int CurrentIndex = 0; CurrentIndex < ReferencingPropertiesMapKeys.Num(); CurrentIndex++)
 			{
-				if (FindDependentArchive.GetReferenceCount(ReferencingPropertiesMapKeys[DependentIndex]) > 0)
+				TouchedThisItteration.Init(false, ReferencingPropertiesMapKeys.Num());
+				FFindReferencersArchive FindDependentArchive(ReferencingPropertiesMapKeys[CurrentIndex], ReferencingPropertiesMapKeys);
+				for (int DependentIndex = CurrentIndex + 1; DependentIndex < ReferencingPropertiesMapKeys.Num(); DependentIndex++)
 				{
-					if (CheckForCircles != DependentIndex)
+					if (!TouchedThisItteration[DependentIndex] && FindDependentArchive.GetReferenceCount(ReferencingPropertiesMapKeys[DependentIndex]) > 0)
 					{
 						UObject* Key = ReferencingPropertiesMapKeys[CurrentIndex];
 						PropertyArrayType Value = ReferencingPropertiesMapValues[CurrentIndex];
@@ -803,12 +803,7 @@ namespace ObjectTools
 						ReferencingPropertiesMapValues[DependentIndex] = Value;
 
 						FindDependentArchive.ResetPotentialReferencer(ReferencingPropertiesMapKeys[CurrentIndex]);
-						
-						if(CheckForCircles == 0)
-						{
-							CheckForCircles = DependentIndex;
-						}
-						
+						TouchedThisItteration[DependentIndex] = true;
 						DependentIndex = CurrentIndex;
 					}
 				}
diff --git a/Engine/Source/Programs/AutomationTool/BuildGraph/Script.cs b/Engine/Source/Programs/AutomationTool/BuildGraph/Script.cs
--- a/Engine/Source/Programs/AutomationTool/BuildGraph/Script.cs
+++ b/Engine/Source/Programs/AutomationTool/BuildGraph/Script.cs
@@ -794,12 +794,10 @@ namespace AutomationTool
 				{
 					if (Types.Length > 0 && Agent.PossibleTypes.Length > 0)
 					{
-						string[] NewTypes = Agent.PossibleTypes.Intersect(Types, StringComparer.InvariantCultureIgnoreCase).ToArray();
-						if (NewTypes.Length == 0)
+						if (Types.Length != Agent.PossibleTypes.Length || !Types.SequenceEqual(Agent.PossibleTypes, StringComparer.InvariantCultureIgnoreCase))
 						{
-							LogError(Element, "No common agent types with previous agent definition");
+							LogError(Element, "Agent types ({0}) were different than previous agent definition with types ({1}). Must either be empty or match exactly.", string.Join(",", Types), string.Join(",", Agent.PossibleTypes));
 						}
-						Agent.PossibleTypes = NewTypes;
 					}
 				}
 				else
diff --git a/Engine/Source/Programs/UnrealBuildTool/Platform/Android/AndroidToolChain.cs b/Engine/Source/Programs/UnrealBuildTool/Platform/Android/AndroidToolChain.cs
--- a/Engine/Source/Programs/UnrealBuildTool/Platform/Android/AndroidToolChain.cs
+++ b/Engine/Source/Programs/UnrealBuildTool/Platform/Android/AndroidToolChain.cs
@@ -1665,12 +1665,71 @@ namespace UnrealBuildTool
 						LinkAction.CommandArguments = String.Format("/c \"{0} {1}\"", LinkAction.CommandPath, LinkAction.CommandArguments);
 						LinkAction.CommandPath = "cmd.exe";
 					}
+
+					// Windows can run into an issue with too long of a commandline when clang tries to call ld to link.
+					// To work around this we call clang to just get the command it would execute and generate a
+					// second response file to directly call ld with the right arguments instead of calling through clang.
+/* disable while tracking down some linker errors this introduces
+					if (!Utils.IsRunningOnMono)
+					{
+						// capture the actual link command without running it
+						ProcessStartInfo StartInfo = new ProcessStartInfo();
+						StartInfo.WorkingDirectory = LinkEnvironment.IntermediateDirectory.FullName;
+						StartInfo.FileName = LinkAction.CommandPath;
+						StartInfo.Arguments = "-### " + LinkAction.CommandArguments;
+						StartInfo.UseShellExecute = false;
+						StartInfo.CreateNoWindow = true;
+						StartInfo.RedirectStandardError = true;
+
+						LinkerCommandline = "";
+
+						Process Proc = new Process();
+						Proc.StartInfo = StartInfo;
+						Proc.ErrorDataReceived += new DataReceivedEventHandler(OutputReceivedForLinker);
+						Proc.Start();
+						Proc.BeginErrorReadLine();
+						Proc.WaitForExit(5000);
+
+						LinkerCommandline = LinkerCommandline.Trim();
+
+						// the command should be in quotes; if not we'll just use clang to link as usual
+						int FirstQuoteIndex = LinkerCommandline.IndexOf('"');
+						if (FirstQuoteIndex >= 0)
+						{
+							int SecondQuoteIndex = LinkerCommandline.Substring(FirstQuoteIndex + 1).IndexOf('"');
+							if (SecondQuoteIndex >= 0)
+							{
+								LinkAction.CommandPath = LinkerCommandline.Substring(FirstQuoteIndex + 1, SecondQuoteIndex - FirstQuoteIndex);
+								LinkAction.CommandArguments = LinkerCommandline.Substring(FirstQuoteIndex + SecondQuoteIndex + 3);
+
+								// replace double backslashes
+								LinkAction.CommandPath = LinkAction.CommandPath.Replace("\\\\", "/");
+
+								// now create a response file for the full command using ld directly
+								FileReference FinalResponseFileName = FileReference.Combine(LinkEnvironment.IntermediateDirectory, OutputFile.Location.GetFileName() + ".responseFinal");
+								FileItem FinalResponseFileItem = FileItem.CreateIntermediateTextFile(FinalResponseFileName, LinkAction.CommandArguments);
+								LinkAction.CommandArguments = string.Format("@\"{0}\"", FinalResponseFileName);
+								LinkAction.PrerequisiteItems.Add(FinalResponseFileItem);
+							}
+						}
+					}
+*/
 				}
 			}
 
 			return Outputs.ToArray();
 		}
 
+		// captures stderr from clang
+		private static string LinkerCommandline = "";
+		static public void OutputReceivedForLinker(Object Sender, DataReceivedEventArgs Line)
+		{
+			if ((Line != null) && (Line.Data != null) && (Line.Data.Contains("--sysroot")))
+			{
+				LinkerCommandline += Line.Data;
+			}
+		}
+
 		private void ExportObjectFilePaths(LinkEnvironment LinkEnvironment, string FileName)
 		{
 			// Write the list of object file directories
diff --git a/Engine/Source/Runtime/Core/Private/HAL/LowLevelMemTracker.cpp b/Engine/Source/Runtime/Core/Private/HAL/LowLevelMemTracker.cpp
--- a/Engine/Source/Runtime/Core/Private/HAL/LowLevelMemTracker.cpp
+++ b/Engine/Source/Runtime/Core/Private/HAL/LowLevelMemTracker.cpp
@@ -61,6 +61,7 @@ DECLARE_LLM_MEMORY_STAT(TEXT("Networking"), STAT_NetworkingLLM, STATGROUP_LLMFUL
 DECLARE_LLM_MEMORY_STAT(TEXT("Meshes"), STAT_MeshesLLM, STATGROUP_LLMFULL);
 DECLARE_LLM_MEMORY_STAT(TEXT("Stats"), STAT_StatsLLM, STATGROUP_LLMFULL);
 DECLARE_LLM_MEMORY_STAT(TEXT("Shaders"), STAT_ShadersLLM, STATGROUP_LLMFULL);
+DECLARE_LLM_MEMORY_STAT(TEXT("PSO"), STAT_PSOLLM, STATGROUP_LLMFULL);
 DECLARE_LLM_MEMORY_STAT(TEXT("Textures"), STAT_TexturesLLM, STATGROUP_LLMFULL);
 DECLARE_LLM_MEMORY_STAT(TEXT("Render Targets"), STAT_RenderTargetsLLM, STATGROUP_LLMFULL);
 DECLARE_LLM_MEMORY_STAT(TEXT("RHIMisc"), STAT_RHIMiscLLM, STATGROUP_LLMFULL);
diff --git a/Engine/Source/Runtime/Core/Public/HAL/LowLevelMemTracker.h b/Engine/Source/Runtime/Core/Public/HAL/LowLevelMemTracker.h
--- a/Engine/Source/Runtime/Core/Public/HAL/LowLevelMemTracker.h
+++ b/Engine/Source/Runtime/Core/Public/HAL/LowLevelMemTracker.h
@@ -138,6 +138,7 @@ enum class ELLMTagSet : uint8
 	macro(Meshes,								"Meshes",						GET_STATFNAME(STAT_MeshesLLM),								GET_STATFNAME(STAT_MeshesSummaryLLM))		\
 	macro(Stats,								"Stats",						GET_STATFNAME(STAT_StatsLLM),								GET_STATFNAME(STAT_EngineSummaryLLM))		\
 	macro(Shaders,								"Shaders",						GET_STATFNAME(STAT_ShadersLLM),								GET_STATFNAME(STAT_EngineSummaryLLM))		\
+	macro(PSO,									"PSO",							GET_STATFNAME(STAT_PSOLLM),									GET_STATFNAME(STAT_EngineSummaryLLM))		\
 	macro(Textures,								"Textures",						GET_STATFNAME(STAT_TexturesLLM),							GET_STATFNAME(STAT_TexturesSummaryLLM))		\
 	macro(RenderTargets,						"RenderTargets",				GET_STATFNAME(STAT_RenderTargetsLLM),						GET_STATFNAME(STAT_EngineSummaryLLM))		\
 	macro(RHIMisc,								"RHIMisc",						GET_STATFNAME(STAT_RHIMiscLLM),								GET_STATFNAME(STAT_EngineSummaryLLM))		\
diff --git a/Engine/Source/Runtime/Engine/Private/StaticMesh.cpp b/Engine/Source/Runtime/Engine/Private/StaticMesh.cpp
--- a/Engine/Source/Runtime/Engine/Private/StaticMesh.cpp
+++ b/Engine/Source/Runtime/Engine/Private/StaticMesh.cpp
@@ -1350,7 +1350,7 @@ FArchive& operator<<(FArchive& Ar, FMeshBuildSettings& BuildSettings)
 // differences, etc.) replace the version GUID below with a new one.
 // In case of merge conflicts with DDC versions, you MUST generate a new GUID
 // and set this new GUID as the version.                                       
-#define STATICMESH_DERIVEDDATA_VER TEXT("42B9CFE0CAFE46F29F43E9049AA21282")
+#define STATICMESH_DERIVEDDATA_VER TEXT("4FDA6C9D166447EAB2DA9EDA8A184591")
 
 static const FString& GetStaticMeshDerivedDataVersion()
 {
diff --git a/Engine/Source/Runtime/RHI/Private/PipelineStateCache.cpp b/Engine/Source/Runtime/RHI/Private/PipelineStateCache.cpp
--- a/Engine/Source/Runtime/RHI/Private/PipelineStateCache.cpp
+++ b/Engine/Source/Runtime/RHI/Private/PipelineStateCache.cpp
@@ -483,6 +483,8 @@ public:
 
 	void DoTask(ENamedThreads::Type CurrentThread, const FGraphEventRef& MyCompletionGraphEvent)
 	{
+		LLM_SCOPE(ELLMTag::PSO);
+
 		if (Pipeline->IsCompute())
 		{
 			FComputePipelineState* ComputePipeline = static_cast<FComputePipelineState*>(Pipeline);
@@ -713,7 +715,7 @@ FRHIComputePipelineState* ExecuteSetComputePipelineState(FComputePipelineState*
 
 FGraphicsPipelineState* PipelineStateCache::GetAndOrCreateGraphicsPipelineState(FRHICommandList& RHICmdList, const FGraphicsPipelineStateInitializer& OriginalInitializer, EApplyRendertargetOption ApplyFlags)
 {
-
+	LLM_SCOPE(ELLMTag::PSO);
 	SCOPE_CYCLE_COUNTER(STAT_GetOrCreatePSO);
 
 	FGraphicsPipelineStateInitializer NewInitializer;
diff --git a/Engine/Source/Runtime/Renderer/Private/PostProcess/PostProcessing.cpp b/Engine/Source/Runtime/Renderer/Private/PostProcess/PostProcessing.cpp
--- a/Engine/Source/Runtime/Renderer/Private/PostProcess/PostProcessing.cpp
+++ b/Engine/Source/Runtime/Renderer/Private/PostProcess/PostProcessing.cpp
@@ -932,11 +932,11 @@ static void AddTemporalAA( FPostprocessContext& Context, FRenderingCompositeOutp
 
 	FSceneViewState* ViewState = Context.View.ViewState;
 
-	FRCPassPostProcessTemporalAA* TemporalAAPass = new(FMemStack::Get()) FRCPassPostProcessTemporalAA(
+	FRCPassPostProcessTemporalAA* TemporalAAPass = Context.Graph.RegisterPass(new(FMemStack::Get()) FRCPassPostProcessTemporalAA(
 		Context,
 		Parameters,
 		Context.View.PrevViewInfo.TemporalAAHistory,
-		&ViewState->PendingPrevFrameViewInfo.TemporalAAHistory);
+		&ViewState->PendingPrevFrameViewInfo.TemporalAAHistory));
 
 	TemporalAAPass->SetInput( ePId_Input0, Context.FinalOutput );
 	TemporalAAPass->SetInput( ePId_Input2, VelocityInput );
@@ -1726,7 +1726,6 @@ void FPostProcessing::Process(FRHICommandListImmediate& RHICmdList, const FViewI
 			// some views don't have a state (thumbnail rendering)
 			if(!GIsHighResScreenshot && View.State && IStereoRendering::IsAPrimaryView(StereoPass))
 			{
-				
 				const bool bUseBasicEyeAdaptation = (AutoExposure.MethodId == EAutoExposureMethod::AEM_Basic);
 
 				if (bUseBasicEyeAdaptation) // log average ps reduction ( non histogram ) 
