Created FunctionalTest base class.

FunctionalTest now implements TreatLogErrorsAsErrors and TreatLogWarningsAsErrors. Default is true and false, but projects and individual tests can override these.

diff --git a/Engine/Source/Developer/FunctionalTesting/Classes/FunctionalTest.h b/Engine/Source/Developer/FunctionalTesting/Classes/FunctionalTest.h
--- a/Engine/Source/Developer/FunctionalTesting/Classes/FunctionalTest.h
+++ b/Engine/Source/Developer/FunctionalTesting/Classes/FunctionalTest.h
@@ -72,7 +72,6 @@ struct FUNCTIONALTESTING_API FPerfStatsRecord
 	bool IsWithinRenderThreadBudget()const;
 };
 
-
 UENUM(BlueprintType)
 enum class EComparisonMethod : uint8
 {
@@ -202,6 +201,18 @@ enum class EFunctionalTestResult : uint8
 	Succeeded
 };
 
+UENUM(BlueprintType)
+enum class EFunctionalTestLogHandling : uint8
+{
+	/**
+	 * How do log categories affect rest results. ProjectDefault can be set in DefaultEngine.ini
+	 * but individual tests can override that
+	 */
+	ProjectDefault,
+	OutputIsError,
+	OutputIgnored
+};
+
 DECLARE_DYNAMIC_MULTICAST_DELEGATE(FFunctionalTestEventSignature);
 DECLARE_DELEGATE_OneParam(FFunctionalTestDoneSignature, class AFunctionalTest*);
 
@@ -226,11 +237,16 @@ protected:
 	uint32 bIsEnabled:1;
 
 	/**
-	 * If this is enabled, any warning logged while this functional test is running is treated as
-	 * an error.
+	 * Determines how LogErrors are handled during this test.
+	 */
+	UPROPERTY(EditAnywhere, BlueprintReadOnly, Category = "Functional Testing", meta = (ScriptName = "LogErrorHandling"))
+	EFunctionalTestLogHandling LogErrorHandling;
+
+	/**
+	 * Determines how LogWarnings are handled during this test.
 	 */
-	UPROPERTY(EditAnywhere, BlueprintReadOnly, Category="Functional Testing", meta=(ScriptName="WarningsAsErrorsValue"))
-	uint32 bWarningsAsErrors:1;
+	UPROPERTY(EditAnywhere, BlueprintReadOnly, Category = "Functional Testing", meta = (ScriptName = "LogWarningHandling"))
+	EFunctionalTestLogHandling LogWarningHandling;
 
 	/**
 	 * The author is the group or person responsible for the test.  Generally you should use a group name
diff --git a/Engine/Source/Developer/FunctionalTesting/Classes/FunctionalTestBase.h b/Engine/Source/Developer/FunctionalTesting/Classes/FunctionalTestBase.h
new file mode 100644
--- /dev/null
+++ b/Engine/Source/Developer/FunctionalTesting/Classes/FunctionalTestBase.h
@@ -0,0 +1,77 @@
+// Copyright 1998-2019 Epic Games, Inc. All Rights Reserved.
+
+#pragma once
+
+#include "CoreMinimal.h"
+#include "Misc/AutomationTest.h"
+
+/**
+ * Base class for Functional test cases.  
+ */
+class FUNCTIONALTESTING_API FFunctionalTestBase : public FAutomationTestBase
+{
+public:
+	FFunctionalTestBase(const FString& InName, const bool bInComplexTask);
+
+	/**
+	 * If true logs will not be included in test events
+	 *
+	 * @return true to suppress logs
+	 */
+	virtual bool SuppressLogs()
+	{
+		return bSuppressLogs || !IsFunctionalTestRunning();
+	}
+
+	/**
+	 * Specify how log errors & warnings should be handled during tests. If values are not set then the project
+	 * defaults will be used.
+	 */
+	void SetLogErrorAndWarningHandling(TOptional<bool> LogErrorsAreErrors, TOptional<bool> LogWarningsAreErrors)
+	{
+		SetLogErrorAndWarningHandlingToDefault();
+	}
+	
+
+	/**
+	 * Determines if Error logs should be considered test errors
+	 */
+	virtual bool TreatLogErrorsAsErrors() override
+	{
+		return bTreatLogErrorsAsErrors;
+	}
+
+	/**
+	 * Determines if Warning logs should be considered test errors
+	 */
+	virtual bool TreatLogWarningsAsErrors() override
+	{
+		return bTreatLogWarningsAsErrors;
+	}
+
+	/**
+	 * Returns true if a functional test is running (does not include map setup)
+	 */
+	bool IsFunctionalTestRunning()
+	{
+		return bIsFunctionalTestRunning;
+	}
+
+	/**
+	 * Marks us as actively running a functional test
+	 */
+	void SetFunctionalTestRunning(bool bIsRunning)
+	{
+		bIsFunctionalTestRunning = bIsRunning;
+	}
+
+protected:
+
+	void SetLogErrorAndWarningHandlingToDefault();
+
+
+	bool bTreatLogErrorsAsErrors;
+	bool bTreatLogWarningsAsErrors;
+	bool bSuppressLogs;
+	bool bIsFunctionalTestRunning;
+};
diff --git a/Engine/Source/Developer/FunctionalTesting/Private/FunctionalTest.cpp b/Engine/Source/Developer/FunctionalTesting/Private/FunctionalTest.cpp
--- a/Engine/Source/Developer/FunctionalTesting/Private/FunctionalTest.cpp
+++ b/Engine/Source/Developer/FunctionalTesting/Private/FunctionalTest.cpp
@@ -24,6 +24,7 @@
 #include "Engine/DebugCameraController.h"
 #include "TraceQueryTestResults.h"
 #include "Misc/RuntimeErrors.h"
+#include "FunctionalTestBase.h"
 
 namespace
 {
@@ -83,10 +84,12 @@ namespace
 	}
 }
 
+
 AFunctionalTest::AFunctionalTest( const FObjectInitializer& ObjectInitializer )
 	: Super(ObjectInitializer)
 	, bIsEnabled(true)
-	, bWarningsAsErrors(false)
+	, LogErrorHandling(EFunctionalTestLogHandling::ProjectDefault)
+	, LogWarningHandling(EFunctionalTestLogHandling::ProjectDefault)
 	, Result(EFunctionalTestResult::Invalid)
 	, PreparationTimeLimit(15.0f)
 	, TimeLimit(60.0f)
@@ -192,7 +195,25 @@ bool AFunctionalTest::RunTest(const TArray<FString>& Params)
 {
 	ensure(GetWorld()->HasBegunPlay());
 
-	FAutomationTestFramework::Get().SetTreatWarningsAsErrors(bWarningsAsErrors);
+	FFunctionalTestBase* FunctionalTest = static_cast<FFunctionalTestBase*>(FAutomationTestFramework::Get().GetCurrentTest());
+	check(FunctionalTest);
+
+	// Set handling of warnings/errors based on this test. Tests can either specify an explicit option or choose to go with the
+	// project defaults.
+	TOptional<bool> bLogErrorsAreErrors, bLogWarningsAreErrors;
+
+	if (LogErrorHandling != EFunctionalTestLogHandling::ProjectDefault)
+	{
+		bLogErrorsAreErrors = LogErrorHandling == EFunctionalTestLogHandling::OutputIsError ? true : false;
+	}
+
+	if (LogWarningHandling != EFunctionalTestLogHandling::ProjectDefault)
+	{
+		bLogWarningsAreErrors = LogWarningHandling == EFunctionalTestLogHandling::OutputIsError ? true : false;
+	}
+
+	FunctionalTest->SetLogErrorAndWarningHandling(bLogErrorsAreErrors, bLogWarningsAreErrors);
+	FunctionalTest->SetFunctionalTestRunning(true);
 
 	FailureMessage = TEXT("");
 	
@@ -301,6 +322,11 @@ bool AFunctionalTest::IsReady_Implementation()
 void AFunctionalTest::FinishTest(EFunctionalTestResult TestResult, const FString& Message)
 {
 	const static UEnum* FTestResultTypeEnum = StaticEnum<EFunctionalTestResult>();
+
+	FFunctionalTestBase* FunctionalTest = static_cast<FFunctionalTestBase*>(FAutomationTestFramework::Get().GetCurrentTest());
+	check(FunctionalTest);
+
+	FunctionalTest->SetFunctionalTestRunning(false);
 	
 	if (bIsRunning == false)
 	{
@@ -308,6 +334,8 @@ void AFunctionalTest::FinishTest(EFunctionalTestResult TestResult, const FString
 		return;
 	}
 
+
+
 	//Force GC at the end of every test.
 	GEngine->ForceGarbageCollection();
 
@@ -369,8 +397,6 @@ void AFunctionalTest::FinishTest(EFunctionalTestResult TestResult, const FString
 	//}
 
 	TestFinishedObserver.ExecuteIfBound(this);
-
-	FAutomationTestFramework::Get().SetTreatWarningsAsErrors(TOptional<bool>());
 }
 
 void AFunctionalTest::EndPlay(const EEndPlayReason::Type EndPlayReason)
diff --git a/Engine/Source/Developer/FunctionalTesting/Private/FunctionalTestBase.cpp b/Engine/Source/Developer/FunctionalTesting/Private/FunctionalTestBase.cpp
new file mode 100644
--- /dev/null
+++ b/Engine/Source/Developer/FunctionalTesting/Private/FunctionalTestBase.cpp
@@ -0,0 +1,23 @@
+// Copyright 1998-2019 Epic Games, Inc. All Rights Reserved.
+
+#include "FunctionalTestBase.h"
+#include "Tests/AutomationTestSettings.h"
+
+
+FFunctionalTestBase::FFunctionalTestBase(const FString& InName, const bool bInComplexTask)
+	: FAutomationTestBase(InName, bInComplexTask)
+{
+	bSuppressLogs = false;
+	bIsFunctionalTestRunning = false;
+	// CDO not available at this point
+	bTreatLogErrorsAsErrors = true;
+	bTreatLogErrorsAsErrors = false;
+}
+
+void FFunctionalTestBase::SetLogErrorAndWarningHandlingToDefault()
+{
+	// Set to project defaults
+	UAutomationTestSettings* Settings = UAutomationTestSettings::StaticClass()->GetDefaultObject<UAutomationTestSettings>();
+	bTreatLogErrorsAsErrors = Settings->bTreatLogErrorsAsTestErrors;
+	bTreatLogWarningsAsErrors = Settings->bTreatLogWarningsAsTestErrors;
+}
diff --git a/Engine/Source/Developer/FunctionalTesting/Private/Tests/ClientFuncTestPerforming.cpp b/Engine/Source/Developer/FunctionalTesting/Private/Tests/ClientFuncTestPerforming.cpp
--- a/Engine/Source/Developer/FunctionalTesting/Private/Tests/ClientFuncTestPerforming.cpp
+++ b/Engine/Source/Developer/FunctionalTesting/Private/Tests/ClientFuncTestPerforming.cpp
@@ -8,6 +8,8 @@
 #include "FunctionalTestingModule.h"
 #include "EngineGlobals.h"
 #include "Tests/AutomationCommon.h"
+#include "FunctionalTestBase.h"
+#include "FunctionalTest.h"
 #include "FunctionalTestingHelper.h"
 
 #if WITH_DEV_AUTOMATION_TESTS
@@ -16,11 +18,11 @@
 
 DEFINE_LOG_CATEGORY_STATIC(LogFunctionalTesting, Log, All);
 
-class FClientFunctionalTestingMapsBase : public FAutomationTestBase
+class FClientFunctionalTestingMapsBase : public FFunctionalTestBase
 {
 public:
 	FClientFunctionalTestingMapsBase(const FString& InName, const bool bInComplexTask)
-		: FAutomationTestBase(InName, bInComplexTask)
+		: FFunctionalTestBase(InName, bInComplexTask)
 	{
 	}
 
@@ -96,6 +98,9 @@ public:
 
 		IFunctionalTestingModule::Get().MarkPendingActivation();
 
+		// Always reset these, even though tests should do the same
+		SetLogErrorAndWarningHandlingToDefault();
+
 		UWorld* TestWorld = GetAnyGameWorld();
 		if (TestWorld && TestWorld->GetMapName() == MapPackageName)
 		{
@@ -129,6 +134,7 @@ public:
 		UE_LOG(LogFunctionalTest, Error, TEXT("Failed to start the %s map (possibly due to BP compilation issues)"), *MapPackageName);
 		return false;
 	}
+
 };
 
 // Runtime tests
