Fix crash when PIEing levels with FTest actor that does not have an active functional test.

diff --git a/Engine/Source/Developer/FunctionalTesting/Private/FunctionalTest.cpp b/Engine/Source/Developer/FunctionalTesting/Private/FunctionalTest.cpp
--- a/Engine/Source/Developer/FunctionalTesting/Private/FunctionalTest.cpp
+++ b/Engine/Source/Developer/FunctionalTesting/Private/FunctionalTest.cpp
@@ -196,7 +196,6 @@ bool AFunctionalTest::RunTest(const TArray<FString>& Params)
 	ensure(GetWorld()->HasBegunPlay());
 
 	FFunctionalTestBase* FunctionalTest = static_cast<FFunctionalTestBase*>(FAutomationTestFramework::Get().GetCurrentTest());
-	check(FunctionalTest);
 
 	// Set handling of warnings/errors based on this test. Tests can either specify an explicit option or choose to go with the
 	// project defaults.
@@ -212,8 +211,11 @@ bool AFunctionalTest::RunTest(const TArray<FString>& Params)
 		bLogWarningsAreErrors = LogWarningHandling == EFunctionalTestLogHandling::OutputIsError ? true : false;
 	}
 
-	FunctionalTest->SetLogErrorAndWarningHandling(bLogErrorsAreErrors, bLogWarningsAreErrors);
-	FunctionalTest->SetFunctionalTestRunning(true);
+	if (FunctionalTest)
+	{
+		FunctionalTest->SetLogErrorAndWarningHandling(bLogErrorsAreErrors, bLogWarningsAreErrors);
+		FunctionalTest->SetFunctionalTestRunning(true);
+	}
 
 	FailureMessage = TEXT("");
 	
@@ -324,9 +326,10 @@ void AFunctionalTest::FinishTest(EFunctionalTestResult TestResult, const FString
 	const static UEnum* FTestResultTypeEnum = StaticEnum<EFunctionalTestResult>();
 
 	FFunctionalTestBase* FunctionalTest = static_cast<FFunctionalTestBase*>(FAutomationTestFramework::Get().GetCurrentTest());
-	check(FunctionalTest);
-
-	FunctionalTest->SetFunctionalTestRunning(false);
+	if (FunctionalTest)
+	{
+		FunctionalTest->SetFunctionalTestRunning(false);
+	}
 	
 	if (bIsRunning == false)
 	{
@@ -334,8 +337,6 @@ void AFunctionalTest::FinishTest(EFunctionalTestResult TestResult, const FString
 		return;
 	}
 
-
-
 	//Force GC at the end of every test.
 	GEngine->ForceGarbageCollection();
 
