Testing - removing a function from the functional testing manager that has no implementation.

diff --git a/Engine/Source/Developer/FunctionalTesting/Classes/FunctionalTestingManager.h b/Engine/Source/Developer/FunctionalTesting/Classes/FunctionalTestingManager.h
--- a/Engine/Source/Developer/FunctionalTesting/Classes/FunctionalTestingManager.h
+++ b/Engine/Source/Developer/FunctionalTesting/Classes/FunctionalTestingManager.h
@@ -54,9 +54,6 @@ class UFunctionalTestingManager : public UBlueprintFunctionLibrary
 	void SetLooped(const bool bNewLooped) { bLooped = bNewLooped; }
 
 	void TickMe(float DeltaTime);
-
-private:
-	void LogMessage(const FString& MessageString, TSharedPtr<IMessageLogListing> LogListing = NULL);
 	
 protected:
 	static UFunctionalTestingManager* GetManager(UObject* WorldContext);
