Fix for UE-68957 - editor crashing when exiting PIE while tests are running

diff --git a/Engine/Source/Developer/FunctionalTesting/Classes/FunctionalTestingManager.h b/Engine/Source/Developer/FunctionalTesting/Classes/FunctionalTestingManager.h
--- a/Engine/Source/Developer/FunctionalTesting/Classes/FunctionalTestingManager.h
+++ b/Engine/Source/Developer/FunctionalTesting/Classes/FunctionalTestingManager.h
@@ -77,6 +77,7 @@ protected:
 	bool bFinished;
 	bool bLooped;
 	bool bInitialDelayApplied;
+    bool bIsTearingDown;
 	uint32 CurrentIteration;
 
 	FFunctionalTestDoneSignature TestFinishedObserver;
diff --git a/Engine/Source/Developer/FunctionalTesting/Private/FunctionalTestingManager.cpp b/Engine/Source/Developer/FunctionalTesting/Private/FunctionalTestingManager.cpp
--- a/Engine/Source/Developer/FunctionalTesting/Private/FunctionalTestingManager.cpp
+++ b/Engine/Source/Developer/FunctionalTesting/Private/FunctionalTestingManager.cpp
@@ -64,6 +64,7 @@ UFunctionalTestingManager::UFunctionalTestingManager( const FObjectInitializer&
 	, bFinished(false)
 	, bLooped(false)
 	, bInitialDelayApplied(false)
+    , bIsTearingDown(false)
 	, CurrentIteration(INDEX_NONE)
 {
 	if (HasAnyFlags(RF_ClassDefaultObject) == false)
@@ -223,6 +224,8 @@ void UFunctionalTestingManager::OnWorldCleanedUp(UWorld* World, bool bSessionEnd
 
 		// Clear the functional test manager once the world is removed.
 		IFunctionalTestingModule::Get().SetManager(nullptr);
+        
+        bIsTearingDown = true;
 	}
 }
 
@@ -255,7 +258,7 @@ void UFunctionalTestingManager::NotifyTestDone(AFunctionalTest* FTest)
 		}
 	}
 
-	if (TestsLeft.Num() > 0 || TestReproStrings.Num() > 0)
+	if ((TestsLeft.Num() > 0 || TestReproStrings.Num() > 0) && !bIsTearingDown)
 	{
 		bIsRunning = RunFirstValidTest();
 	}
