Fixed issue with Gameplay tests timing out.

Fixed issue with timeouts not being reported as failures.

diff --git a/Engine/Source/Developer/FunctionalTesting/Private/FunctionalTest.cpp b/Engine/Source/Developer/FunctionalTesting/Private/FunctionalTest.cpp
--- a/Engine/Source/Developer/FunctionalTesting/Private/FunctionalTest.cpp
+++ b/Engine/Source/Developer/FunctionalTesting/Private/FunctionalTest.cpp
@@ -323,54 +323,30 @@ bool AFunctionalTest::IsReady_Implementation()
 
 void AFunctionalTest::FinishTest(EFunctionalTestResult TestResult, const FString& Message)
 {
-	const static UEnum* FTestResultTypeEnum = StaticEnum<EFunctionalTestResult>();
-
-	FFunctionalTestBase* FunctionalTest = static_cast<FFunctionalTestBase*>(FAutomationTestFramework::Get().GetCurrentTest());
-	if (FunctionalTest)
-	{
-		FunctionalTest->SetFunctionalTestRunning(false);
-	}
-	
 	if (bIsRunning == false)
 	{
 		// ignore
 		return;
 	}
-
-	//Force GC at the end of every test.
-	GEngine->ForceGarbageCollection();
-
+	
+	// Do reporting first. When we start cleaning things up internal states that capture results
+	// are reset.
+	
 	Result = TestResult;
-
-	bIsRunning = false;
-	SetActorTickEnabled(false);
-
-	OnTestFinished.Broadcast();
-
-	AActor** ActorToDestroy = AutoDestroyActors.GetData();
-
-	for (int32 ActorIndex = 0; ActorIndex < AutoDestroyActors.Num(); ++ActorIndex, ++ActorToDestroy)
-	{
-		if (*ActorToDestroy != NULL)
-		{
-			// will be removed next frame
-			(*ActorToDestroy)->SetLifeSpan( 0.01f );
-		}
-	}
-
+	
+	const static UEnum* FTestResultTypeEnum = StaticEnum<EFunctionalTestResult>();
+	
 	const FText ResultText = FTestResultTypeEnum->GetDisplayNameTextByValue( (int64)TestResult );
-
+	
 	//Output map and test name along with results
 	UWorld* World = GetWorld();
 	FString WorldName = (World ? UWorld::RemovePIEPrefix(World->GetMapName()) : "");
 	const FString OutMessage = FString::Printf(TEXT("%s %s %s: \"%s\"")
-		, *WorldName
-		, *GetName()
-		, *ResultText.ToString()
-		, Message.IsEmpty() == false ? *Message : TEXT("Test finished"));
-
-	AutoDestroyActors.Reset();
-		
+											   , *WorldName
+											   , *GetName()
+											   , *ResultText.ToString()
+											   , Message.IsEmpty() == false ? *Message : TEXT("Test finished"));
+	
 	switch (TestResult)
 	{
 		case EFunctionalTestResult::Invalid:
@@ -379,18 +355,45 @@ void AFunctionalTest::FinishTest(EFunctionalTestResult TestResult, const FString
 			UE_VLOG(this, LogFunctionalTest, Error, TEXT("%s"), *OutMessage);
 			UE_LOG(LogFunctionalTest, Error, TEXT("%s"), *OutMessage);
 			break;
-
+			
 		case EFunctionalTestResult::Running:
 			UE_VLOG(this, LogFunctionalTest, Warning, TEXT("%s"), *OutMessage);
 			UE_LOG(LogFunctionalTest, Warning, TEXT("%s"), *OutMessage);
 			break;
-		
+			
 		default:
 			UE_VLOG(this, LogFunctionalTest, Log, TEXT("%s"), *OutMessage);
 			UE_LOG(LogFunctionalTest, Log, TEXT("%s"), *OutMessage);
 			break;
 	}
 	
+	FFunctionalTestBase* FunctionalTest = static_cast<FFunctionalTestBase*>(FAutomationTestFramework::Get().GetCurrentTest());
+	if (FunctionalTest)
+	{
+		FunctionalTest->SetFunctionalTestRunning(false);
+	}
+
+	bIsRunning = false;
+	SetActorTickEnabled(false);
+
+	OnTestFinished.Broadcast();
+
+	AActor** ActorToDestroy = AutoDestroyActors.GetData();
+
+	for (int32 ActorIndex = 0; ActorIndex < AutoDestroyActors.Num(); ++ActorIndex, ++ActorToDestroy)
+	{
+		if (*ActorToDestroy != NULL)
+		{
+			// will be removed next frame
+			(*ActorToDestroy)->SetLifeSpan( 0.01f );
+		}
+	}
+	
+	AutoDestroyActors.Reset();
+	
+	//Force GC at the end of every test.
+	GEngine->ForceGarbageCollection();
+	
 	//if (AdditionalDetails.IsEmpty() == false)
 	//{
 	//	const FString AdditionalDetails = FString::Printf(TEXT("%s %s, time %.2fs"), *GetAdditionalTestFinishedMessage(TestResult), *OnAdditionalTestFinishedMessageRequest(TestResult), TotalTime);
