Copy from Dev-Enterprise - Add capacity to Black list any test
+ Add copy test names contextual menu

diff --git a/Engine/Source/Developer/AutomationWindow/Private/SAutomationTestItemContextMenu.h b/Engine/Source/Developer/AutomationWindow/Private/SAutomationTestItemContextMenu.h
--- a/Engine/Source/Developer/AutomationWindow/Private/SAutomationTestItemContextMenu.h
+++ b/Engine/Source/Developer/AutomationWindow/Private/SAutomationTestItemContextMenu.h
@@ -34,9 +34,10 @@ public:
 	 * @param InArgs The declaration data for this widget.
 	 * @param InSessionManager The session to use.
 	 */
-	void Construct( const FArguments& InArgs, const TArray<FString>& InAssetNames )
+	void Construct( const FArguments& InArgs, const TArray<FString>& InAssetNames, const FString InTestNames )
 	{
 		AssetNames = InAssetNames;
+		TestNames = InTestNames;
 
 		ChildSlot
 		[
@@ -62,7 +63,14 @@ protected:
 
 		MenuBuilder.BeginSection("AutomationOptions", LOCTEXT("MenuHeadingText", "Automation Options"));
 		{
-			MenuBuilder.AddMenuEntry(LOCTEXT("AutomationMenuEntryLoadText", "Load the asset(s)"), FText::GetEmpty(), FSlateIcon(), FUIAction(FExecuteAction::CreateRaw(this, &SAutomationTestItemContextMenu::HandleContextItemTerminate)));
+			if (!TestNames.IsEmpty())
+			{
+				MenuBuilder.AddMenuEntry(LOCTEXT("AutomationMenuEntryCopyTestNameText", "Copy test name(s)"), FText::GetEmpty(), FSlateIcon(), FUIAction(FExecuteAction::CreateRaw(this, &SAutomationTestItemContextMenu::HandleContextItemCopyName)));
+			}
+			if (AssetNames.Num())
+			{
+				MenuBuilder.AddMenuEntry(LOCTEXT("AutomationMenuEntryLoadText", "Load the asset(s)"), FText::GetEmpty(), FSlateIcon(), FUIAction(FExecuteAction::CreateRaw(this, &SAutomationTestItemContextMenu::HandleContextItemTerminate)));
+			}
 		}
 		MenuBuilder.EndSection();
 
@@ -80,10 +88,19 @@ private:
 		}
 	}
 
+	/** Handle the context menu closing down. Copy the test names to clipboard */
+	void HandleContextItemCopyName()
+	{
+		FPlatformApplicationMisc::ClipboardCopy(*TestNames);
+	}
+
 private:
 
 	/** Holds the selected asset name. */
 	TArray<FString> AssetNames;
+
+	//** Holds the test names. */
+	FString TestNames;
 };
 
 
diff --git a/Engine/Source/Developer/AutomationWindow/Private/SAutomationWindow.cpp b/Engine/Source/Developer/AutomationWindow/Private/SAutomationWindow.cpp
--- a/Engine/Source/Developer/AutomationWindow/Private/SAutomationWindow.cpp
+++ b/Engine/Source/Developer/AutomationWindow/Private/SAutomationWindow.cpp
@@ -1243,18 +1243,29 @@ TSharedPtr<SWidget> SAutomationWindow::HandleAutomationListContextMenuOpening()
 {
 	TArray< TSharedPtr<IAutomationReport> >SelectedReport = TestTable->GetSelectedItems();
 
+	TArray<FString> TestNames;
 	TArray<FString> AssetNames;
-	for (int32 ReportIndex = 0; ReportIndex < SelectedReport.Num(); ++ReportIndex)
+	for (TSharedPtr<IAutomationReport> Report : SelectedReport)
 	{
-		// TODO This is super sketch, we were interpreting the parameter always as the asset, this is no good.
-		if (SelectedReport[ReportIndex].IsValid() && (SelectedReport[ReportIndex]->GetTestParameter().Len() > 0))
+		if (Report.IsValid())
 		{
-			AssetNames.Add(SelectedReport[ReportIndex]->GetTestParameter());
+			FString TestName = Report->GetFullTestPath();
+			if (!TestName.IsEmpty())
+			{
+				TestNames.Add(TestName);
+			}
+			FString Param = Report->GetTestParameter();
+			if (Param.Len() > 0 && Param.StartsWith(TEXT("/")))
+			{
+				// Assume that if parameter start with a "/", it should be an asset
+				AssetNames.Add(Param);
+			}
 		}
 	}		
-	if (AssetNames.Num())
+	
+	if (AssetNames.Num() || TestNames.Num())
 	{
-		return SNew(SAutomationTestItemContextMenu, AssetNames);
+		return SNew(SAutomationTestItemContextMenu, AssetNames, FString::Join(TestNames, TEXT("\n")));
 	}
 
 	return nullptr;
diff --git a/Engine/Source/Developer/FunctionalTesting/Private/FunctionalTestingModule.cpp b/Engine/Source/Developer/FunctionalTesting/Private/FunctionalTestingModule.cpp
--- a/Engine/Source/Developer/FunctionalTesting/Private/FunctionalTestingModule.cpp
+++ b/Engine/Source/Developer/FunctionalTesting/Private/FunctionalTestingModule.cpp
@@ -38,21 +38,7 @@ private:
 	UWorld* GetTestWorld();
 	void OnGetAssetTagsForWorld(const UWorld* World, TArray<UObject::FAssetRegistryTag>& OutTags);
 
-	void BuildTestBlacklistFromConfig();
-	bool IsBlacklisted(const FString& MapName, const FString& TestName, FString* OutReason=nullptr, bool* OutWarn=nullptr) const;
 
-	struct FBlacklistEntry
-	{
-		FBlacklistEntry() :
-			bWarn(false) {}
-
-		FString	Map;
-		FString Test;
-		FString Reason;
-		bool bWarn;
-	};
-
-	TMap<FString, FBlacklistEntry> TestBlacklist;
 	TWeakObjectPtr<class UFunctionalTestingManager> TestManager;
 	bool bPendingActivation;
 };
@@ -63,7 +49,6 @@ void FFunctionalTestingModule::StartupModule()
 #if WITH_EDITOR
 	FWorldDelegates::GetAssetTags.AddRaw(this, &FFunctionalTestingModule::OnGetAssetTagsForWorld);
 #endif
-	BuildTestBlacklistFromConfig();
 }
 
 void FFunctionalTestingModule::ShutdownModule() 
@@ -160,25 +145,9 @@ void FFunctionalTestingModule::GetMapTests(bool bEditorOnlyTests, TArray<FString
 
 							if (MapTest.Split(TEXT("|"), &BeautifulTestName, &RealTestName))
 							{
-								FString BlacklistReason;
-								bool bWarn(false);
-								if (!IsBlacklisted(MapPackageName, RealTestName, &BlacklistReason, &bWarn))
-								{
-									OutBeautifiedNames.Add(MapPackageName + TEXT(".") + *BeautifulTestName);
-									OutTestCommands.Add(MapAssetPath + TEXT(";") + MapAsset.PackageName.ToString() + TEXT(";") + *RealTestName);
-									OutTestMapAssets.AddUnique(MapAssetPath);
-								}
-								else
-								{
-									if (bWarn)
-									{
-										UE_LOG(LogFunctionalTest, Warning, TEXT("Test '%s' is blacklisted. %s"), *MapTest, *BlacklistReason);
-									}
-									else
-									{
-										UE_LOG(LogFunctionalTest, Display, TEXT("Test '%s' is blacklisted. %s"), *MapTest, *BlacklistReason);
-									}
-								}
+								OutBeautifiedNames.Add(MapPackageName + TEXT(".") + *BeautifulTestName);
+								OutTestCommands.Add(MapAssetPath + TEXT(";") + MapAsset.PackageName.ToString() + TEXT(";") + *RealTestName);
+								OutTestMapAssets.AddUnique(MapAssetPath);
 							}
 						}
 					}
@@ -286,99 +255,6 @@ void FFunctionalTestingModule::RunTestOnMap(const FString& TestName, bool bClear
 	}
 }
 
-void FFunctionalTestingModule::BuildTestBlacklistFromConfig() 
-{
-	TestBlacklist.Empty();
-	if (GConfig)
-	{
-	
-		const FString CommandLine = FCommandLine::Get();
-
-		for (const TPair<FString,FConfigFile>& Config : *GConfig)
-		{
-			FConfigSection* BlacklistSection = GConfig->GetSectionPrivate(TEXT("AutomationTestBlacklist"), false, true, Config.Key);
-			if (BlacklistSection)
-			{
-				// Parse all blacklist definitions of the format "BlacklistTest=(Map=/Game/Tests/MapName, Test=TestName, Reason="Foo")"
-				for (FConfigSection::TIterator Section(*BlacklistSection); Section; ++Section)
-				{
-					if (Section.Key() == TEXT("BlacklistTest"))
-					{
-						FString BlacklistValue = Section.Value().GetValue();
-						FString Map, Test, Reason, Warn;
-						bool bSuccess = false;
-						
-						if (FParse::Value(*BlacklistValue, TEXT("Map="), Map, true) && FParse::Value(*BlacklistValue, TEXT("Test="), Test, true))
-						{
-							FParse::Value(*BlacklistValue, TEXT("Reason="), Reason);
-							FParse::Value(*BlacklistValue, TEXT("Warn="), Warn);
-
-							// These are used as folders so ensure they match the expected layout
-							if (Map.StartsWith(TEXT("/")))
-							{
-								FString ListName = Map + TEXT("/") + Test;
-								ListName.RemoveSpacesInline();
-								bSuccess = true;
-
-								if (CommandLine.Contains(Map) || CommandLine.Contains(Test))
-								{
-									UE_LOG(LogFunctionalTest, Warning, TEXT("Test '%s' is blacklisted but allowing due to command line."), *BlacklistValue);
-								}
-								else
-								{
-									// convert Pretty.Name into PrettyName as we compare on the latter.
-									ListName = ListName.Replace(TEXT("."), TEXT(""));
-									FBlacklistEntry& Entry = TestBlacklist.Add(ListName);
-									Entry.Map = Map;
-									Entry.Test = Test;
-									Entry.Reason = Reason;
-									Entry.bWarn = Warn.ToBool();
-								}
-							}
-						}
-						
-						if (!bSuccess)
-						{
-							UE_LOG(LogFunctionalTest, Error, TEXT("Invalid blacklisted test definition: '%s'"), *BlacklistValue);
-						}
-					}
-				}
-			}
-		}
-	}
-
-	if (TestBlacklist.Num() > 0)
-	{
-		UE_LOG(LogFunctionalTest, Log, TEXT("Automated Test Blacklist:"));
-		for (auto& KV : TestBlacklist)
-		{
-			UE_LOG(LogFunctionalTest, Log, TEXT("\tTest: %s"), *KV.Key);
-		}
-	}
-}
-
-bool FFunctionalTestingModule::IsBlacklisted(const FString& MapName, const FString& TestName, FString* OutReason, bool *OutWarn) const
-{
-	FString ListName = MapName + TEXT("/") + TestName;
-	ListName.RemoveSpacesInline();
-
-	const FBlacklistEntry* Entry = TestBlacklist.Find(ListName);
-
-	if (Entry)
-	{
-		if (OutReason != nullptr)
-		{
-			*OutReason = Entry->Reason;
-		}
-
-		if (OutWarn != nullptr)
-		{
-			*OutWarn = Entry->bWarn;
-		}
-	}
-
-	return Entry != nullptr;
-}
 
 //////////////////////////////////////////////////////////////////////////
 // Exec
diff --git a/Engine/Source/Runtime/AutomationWorker/Private/AutomationWorkerModule.cpp b/Engine/Source/Runtime/AutomationWorker/Private/AutomationWorkerModule.cpp
--- a/Engine/Source/Runtime/AutomationWorker/Private/AutomationWorkerModule.cpp
+++ b/Engine/Source/Runtime/AutomationWorker/Private/AutomationWorkerModule.cpp
@@ -40,6 +40,8 @@ void FAutomationWorkerModule::StartupModule()
 
 	FAutomationTestFramework::Get().PreTestingEvent.AddRaw(this, &FAutomationWorkerModule::HandlePreTestingEvent);
 	FAutomationTestFramework::Get().PostTestingEvent.AddRaw(this, &FAutomationWorkerModule::HandlePostTestingEvent);
+
+	FAutomationTestFramework::Get().BuildTestBlacklistFromConfig();
 }
 
 void FAutomationWorkerModule::ShutdownModule()
diff --git a/Engine/Source/Runtime/Core/Private/Misc/AutomationTest.cpp b/Engine/Source/Runtime/Core/Private/Misc/AutomationTest.cpp
--- a/Engine/Source/Runtime/Core/Private/Misc/AutomationTest.cpp
+++ b/Engine/Source/Runtime/Core/Private/Misc/AutomationTest.cpp
@@ -396,6 +396,105 @@ void FAutomationTestFramework::LoadTestModules( )
 	}
 }
 
+void FAutomationTestFramework::BuildTestBlacklistFromConfig()
+{
+	TestBlacklist.Empty();
+	if (GConfig)
+	{
+
+		const FString CommandLine = FCommandLine::Get();
+
+		for (const TPair<FString, FConfigFile>& Config : *GConfig)
+		{
+			FConfigSection* BlacklistSection = GConfig->GetSectionPrivate(TEXT("AutomationTestBlacklist"), false, true, Config.Key);
+			if (BlacklistSection)
+			{
+				// Parse all blacklist definitions of the format "BlacklistTest=(Map=/Game/Tests/MapName, Test=TestName, Reason="Foo")"
+				for (FConfigSection::TIterator Section(*BlacklistSection); Section; ++Section)
+				{
+					if (Section.Key() == TEXT("BlacklistTest"))
+					{
+						FString BlacklistValue = Section.Value().GetValue();
+						FString Map, Test, Reason, Warn, ListName;
+						bool bSuccess = false;
+
+						if (FParse::Value(*BlacklistValue, TEXT("Test="), Test, true))
+						{
+							ListName = FString(Test);
+							FParse::Value(*BlacklistValue, TEXT("Map="), Map, true);
+							FParse::Value(*BlacklistValue, TEXT("Reason="), Reason);
+							FParse::Value(*BlacklistValue, TEXT("Warn="), Warn);
+
+							if (Map.IsEmpty())
+							{
+								// Test with no Map property
+								bSuccess = true;
+							}
+							else if (Map.StartsWith(TEXT("/")))
+							{
+								// Account for Functional Tests based on Map - historically blacklisting was made only for functional tests
+								ListName = TEXT("Project.Functional Tests.") + Map + TEXT(".") + ListName;
+								bSuccess = true;
+							}
+
+						}
+
+						if (bSuccess)
+						{
+							if ((!Map.IsEmpty() && CommandLine.Contains(Map)) || CommandLine.Contains(Test))
+							{
+								UE_LOG(LogAutomationTest, Warning, TEXT("Test '%s' is blacklisted but allowing due to command line."), *BlacklistValue);
+							}
+							else
+							{
+								FBlacklistEntry& Entry = TestBlacklist.Add(ListName);
+								Entry.Map = Map;
+								Entry.Test = Test;
+								Entry.Reason = Reason;
+								Entry.bWarn = Warn.ToBool();
+							}
+						}
+						else
+						{
+							UE_LOG(LogAutomationTest, Error, TEXT("Invalid blacklisted test definition: '%s'"), *BlacklistValue);
+						}
+					}
+				}
+			}
+		}
+	}
+
+	if (TestBlacklist.Num() > 0)
+	{
+		UE_LOG(LogAutomationTest, Log, TEXT("Automated Test Blacklist:"));
+		for (auto& KV : TestBlacklist)
+		{
+			UE_LOG(LogAutomationTest, Log, TEXT("\tTest: %s"), *KV.Key);
+		}
+	}
+}
+
+bool FAutomationTestFramework::IsBlacklisted(const FString& TestName, FString* OutReason, bool *OutWarn) const
+{
+	const FBlacklistEntry* Entry = TestBlacklist.Find(TestName);
+
+	if (Entry)
+	{
+		if (OutReason != nullptr)
+		{
+			*OutReason = Entry->Reason;
+		}
+
+		if (OutWarn != nullptr)
+		{
+			*OutWarn = Entry->bWarn;
+		}
+	}
+
+	return Entry != nullptr;
+}
+
+
 void FAutomationTestFramework::GetValidTestNames( TArray<FAutomationTestInfo>& TestInfo ) const
 {
 	TestInfo.Empty();
@@ -461,7 +560,30 @@ void FAutomationTestFramework::GetValidTestNames( TArray<FAutomationTestInfo>& T
 		
 		if (bEnabled && bPassesApplicationRequirements && bPassesFeatureRequirements && bPassesFilterRequirement)
 		{
-			CurTest->GenerateTestNames(TestInfo);
+			TArray<FAutomationTestInfo> TestsToAdd;
+			CurTest->GenerateTestNames(TestsToAdd);
+			for (FAutomationTestInfo Test : TestsToAdd)
+			{
+				FString BlacklistReason;
+				bool bWarn(false);
+				FString TestName = Test.GetDisplayName();
+				if (!IsBlacklisted(TestName, &BlacklistReason, &bWarn))
+				{
+					TestInfo.Add(Test);
+				}
+				else
+				{
+					if (bWarn)
+					{
+						UE_LOG(LogAutomationTest, Warning, TEXT("Test '%s' is blacklisted. %s"), *TestName, *BlacklistReason);
+					}
+					else
+					{
+						UE_LOG(LogAutomationTest, Display, TEXT("Test '%s' is blacklisted. %s"), *TestName, *BlacklistReason);
+					}
+				}
+			}
+			
 		}
 
 		// Make sure people are not writing complex tests that take forever to return the names of the tests
diff --git a/Engine/Source/Runtime/Core/Public/Misc/AutomationTest.h b/Engine/Source/Runtime/Core/Public/Misc/AutomationTest.h
--- a/Engine/Source/Runtime/Core/Public/Misc/AutomationTest.h
+++ b/Engine/Source/Runtime/Core/Public/Misc/AutomationTest.h
@@ -792,6 +792,11 @@ public:
 	 */
 	void LoadTestModules();
 
+	/**
+	 * Load the test Blacklist from the config.
+	 */
+	void BuildTestBlacklistFromConfig();
+
 	/**
 	 * Populates the provided array with the names of all tests in the framework that are valid to run for the current
 	 * application settings.
@@ -904,7 +909,19 @@ private:
 		class FAutomationTestBase* CurTest;
 	};
 
-	 friend class FAutomationTestFeedbackContext;
+	//** Store information about blacklisted test */
+	struct FBlacklistEntry
+	{
+		FBlacklistEntry() :
+			bWarn(false) {}
+
+		FString Map;
+		FString Test;
+		FString Reason;
+		bool bWarn;
+	};
+
+	friend class FAutomationTestFeedbackContext;
 	/** Helper method called to prepare settings for automation testing to follow */
 	void PrepForAutomationTests();
 
@@ -936,6 +953,16 @@ private:
 	bool InternalStopTest(FAutomationTestExecutionInfo& OutExecutionInfo);
 
 
+	/**
+	 * Internal helper method that verify if a test is black listed.
+	 *
+	 * @param	TestName		Beautified test name to be checked
+	 * @param	OutReason		Output the reason for the test being blacklisted
+	 * @param	OutWarn			Output true if the config ask for a warning message
+	 * @return	true if the TestName is part of the blacklist.
+	 */
+	bool IsBlacklisted(const FString& TestName, FString* OutReason = nullptr, bool* OutWarn = nullptr) const;
+
 	/** Constructor */
 	FAutomationTestFramework();
 
@@ -986,6 +1013,8 @@ private:
 	bool bForceSmokeTests;
 
 	bool bCaptureStack;
+
+	TMap<FString, FBlacklistEntry> TestBlacklist;
 };
 
 
