Add property to blacklist tests based on RHI from Automated Tests

diff --git a/Engine/Source/Developer/FunctionalTesting/Private/FunctionalTestingModule.cpp b/Engine/Source/Developer/FunctionalTesting/Private/FunctionalTestingModule.cpp
--- a/Engine/Source/Developer/FunctionalTesting/Private/FunctionalTestingModule.cpp
+++ b/Engine/Source/Developer/FunctionalTesting/Private/FunctionalTestingModule.cpp
@@ -2,6 +2,7 @@
 
 #include "FunctionalTestingModule.h"
 #include "FunctionalTestingManager.h"
+#include "Misc/App.h"
 #include "Misc/CoreMisc.h"
 #include "Misc/ConfigCacheIni.h"
 #include "Engine/World.h"
@@ -49,6 +50,7 @@ private:
 		FString	Map;
 		FString Test;
 		FString Reason;
+		TArray<FString> RHIs;
 		bool bWarn;
 	};
 
@@ -305,13 +307,14 @@ void FFunctionalTestingModule::BuildTestBlacklistFromConfig()
 					if (Section.Key() == TEXT("BlacklistTest"))
 					{
 						FString BlacklistValue = Section.Value().GetValue();
-						FString Map, Test, Reason, Warn;
+						FString Map, Test, Reason, Warn, RHIs;
 						bool bSuccess = false;
 						
 						if (FParse::Value(*BlacklistValue, TEXT("Map="), Map, true) && FParse::Value(*BlacklistValue, TEXT("Test="), Test, true))
 						{
 							FParse::Value(*BlacklistValue, TEXT("Reason="), Reason);
 							FParse::Value(*BlacklistValue, TEXT("Warn="), Warn);
+							FParse::Value(*BlacklistValue, TEXT("RHIs="), RHIs);
 
 							// These are used as folders so ensure they match the expected layout
 							if (Map.StartsWith(TEXT("/")))
@@ -332,6 +335,14 @@ void FFunctionalTestingModule::BuildTestBlacklistFromConfig()
 									Entry.Map = Map;
 									Entry.Test = Test;
 									Entry.Reason = Reason;
+									if (!RHIs.IsEmpty())
+									{
+										RHIs.ToLower().ParseIntoArray(Entry.RHIs, TEXT(","), true);
+										for (int32 Index = 0; Index != Entry.RHIs.Num(); ++Index)
+										{
+											Entry.RHIs[Index] = Entry.RHIs[Index].Trim();
+										}
+									}
 									Entry.bWarn = Warn.ToBool();
 								}
 							}
@@ -366,6 +377,11 @@ bool FFunctionalTestingModule::IsBlacklisted(const FString& MapName, const FStri
 
 	if (Entry)
 	{
+		if (Entry->RHIs.Num() != 0 && !Entry->RHIs.Contains(FApp::GetGraphicsRHI().ToLower()))
+		{
+			return false;
+		}
+
 		if (OutReason != nullptr)
 		{
 			*OutReason = Entry->Reason;
