Fixing FString::Trim warnings

diff --git a/Engine/Source/Developer/FunctionalTesting/Private/FunctionalTestingModule.cpp b/Engine/Source/Developer/FunctionalTesting/Private/FunctionalTestingModule.cpp
--- a/Engine/Source/Developer/FunctionalTesting/Private/FunctionalTestingModule.cpp
+++ b/Engine/Source/Developer/FunctionalTesting/Private/FunctionalTestingModule.cpp
@@ -340,7 +340,7 @@ void FFunctionalTestingModule::BuildTestBlacklistFromConfig()
 										RHIs.ToLower().ParseIntoArray(Entry.RHIs, TEXT(","), true);
 										for (int32 Index = 0; Index != Entry.RHIs.Num(); ++Index)
 										{
-											Entry.RHIs[Index] = Entry.RHIs[Index].Trim();
+											Entry.RHIs[Index].TrimStartAndEndInline();
 										}
 									}
 									Entry.bWarn = Warn.ToBool();
