Fix non editor build

diff --git a/Engine/Source/Developer/FunctionalTesting/Private/AutomationBlueprintFunctionLibrary.cpp b/Engine/Source/Developer/FunctionalTesting/Private/AutomationBlueprintFunctionLibrary.cpp
--- a/Engine/Source/Developer/FunctionalTesting/Private/AutomationBlueprintFunctionLibrary.cpp
+++ b/Engine/Source/Developer/FunctionalTesting/Private/AutomationBlueprintFunctionLibrary.cpp
@@ -42,7 +42,9 @@
 #include "SceneView.h"
 #include "Engine/GameEngine.h"
 #include "Engine/LevelStreaming.h"
+#if WITH_EDITOR
 #include "SLevelViewport.h"
+#endif
 
 #define LOCTEXT_NAMESPACE "Automation"
 
