More complete fix for UE-89839

Added a high level define (WITH_AUTOMATION_TESTS) that is set when thier WITH_DEV_AUTOMATION_TESTS or WITH_PERF_AUTOMATION_TESTS is set).

This allows support code to be wrapped in a single define and the WITH_DEV / WITH_PERF macros used for simply including/excluding test definitions.

#tests compiled UE4Game in dev/test/shipping

diff --git a/Engine/Source/Developer/FunctionalTesting/Classes/ScreenshotFunctionalTestBase.h b/Engine/Source/Developer/FunctionalTesting/Classes/ScreenshotFunctionalTestBase.h
--- a/Engine/Source/Developer/FunctionalTesting/Classes/ScreenshotFunctionalTestBase.h
+++ b/Engine/Source/Developer/FunctionalTesting/Classes/ScreenshotFunctionalTestBase.h
@@ -73,7 +73,7 @@ protected:
 
 	FIntPoint ViewportRestoreSize;
 
-#if (WITH_DEV_AUTOMATION_TESTS || WITH_PERF_AUTOMATION_TESTS)
+#if WITH_AUTOMATION_TESTS
 	TSharedPtr<FAutomationTestScreenshotEnvSetup> ScreenshotEnvSetup;
 #endif
 
diff --git a/Engine/Source/Developer/FunctionalTesting/Private/AutomationBlueprintFunctionLibrary.cpp b/Engine/Source/Developer/FunctionalTesting/Private/AutomationBlueprintFunctionLibrary.cpp
--- a/Engine/Source/Developer/FunctionalTesting/Private/AutomationBlueprintFunctionLibrary.cpp
+++ b/Engine/Source/Developer/FunctionalTesting/Private/AutomationBlueprintFunctionLibrary.cpp
@@ -82,7 +82,7 @@ bool UAutomationEditorTask::IsTaskDone() const
 	return IsValidTask() && Task->IsDone();
 }
 
-#if (WITH_DEV_AUTOMATION_TESTS || WITH_PERF_AUTOMATION_TESTS)
+#if WITH_AUTOMATION_TESTS
 
 template<typename T>
 FConsoleVariableSwapperTempl<T>::FConsoleVariableSwapperTempl(FString InConsoleVariableName)
@@ -582,7 +582,7 @@ private:
 	FAutomationScreenshotOptions Options;
 };
 
-#endif
+#endif // WITH_AUTOMATION_TESTS
 
 class FScreenshotTakenState : public FAutomationTaskStatusBase
 {
@@ -604,7 +604,7 @@ public:
 
 	virtual ~FScreenshotTakenState()
 	{
-#if (WITH_DEV_AUTOMATION_TESTS || WITH_PERF_AUTOMATION_TESTS)
+#if WITH_AUTOMATION_TESTS
 		FAutomationTestFramework::Get().OnScreenshotCompared.RemoveAll(this);
 #endif
 		FScreenshotRequest::OnScreenshotRequestProcessed().RemoveAll(this);
@@ -699,7 +699,7 @@ bool UAutomationBlueprintFunctionLibrary::TakeAutomationScreenshotInternal(UObje
 {
 	UAutomationBlueprintFunctionLibrary::FinishLoadingBeforeScreenshot();
 
-#if (WITH_DEV_AUTOMATION_TESTS || WITH_PERF_AUTOMATION_TESTS)
+#if WITH_AUTOMATION_TESTS
 	FAutomationScreenshotTaker* TempObject = new FAutomationScreenshotTaker(WorldContextObject ? WorldContextObject->GetWorld() : nullptr, Name, Notes, Options);
 #endif
 
@@ -779,7 +779,7 @@ bool UAutomationBlueprintFunctionLibrary::TakeAutomationScreenshotOfUI_Immediate
 				FIntVector OutSize;
 				if (FSlateApplication::Get().TakeScreenshot(Viewport.ToSharedRef(), OutColorData, OutSize))
 				{
-#if (WITH_DEV_AUTOMATION_TESTS || WITH_PERF_AUTOMATION_TESTS)
+#if WITH_AUTOMATION_TESTS
 					// For UI, we only care about what the final image looks like. So don't compare alpha channel.
 					// In editor, scene is rendered into a PF_B8G8R8A8 RT and then copied over to the R10B10G10A2 swapchain back buffer and
 					// this copy ignores alpha. In game, however, scene is directly rendered into the back buffer and the alpha values are
@@ -1079,7 +1079,7 @@ UAutomationEditorTask* UAutomationBlueprintFunctionLibrary::TakeHighResScreensho
 
 			FinishLoadingBeforeScreenshot();
 
-#if (WITH_DEV_AUTOMATION_TESTS || WITH_PERF_AUTOMATION_TESTS)
+#if WITH_AUTOMATION_TESTS
 			if (GIsAutomationTesting)
 			{
 				if (FAutomationTestBase* CurrentTest = FAutomationTestFramework::Get().GetCurrentTest())
diff --git a/Engine/Source/Developer/FunctionalTesting/Private/ScreenshotFunctionalTestBase.cpp b/Engine/Source/Developer/FunctionalTesting/Private/ScreenshotFunctionalTestBase.cpp
--- a/Engine/Source/Developer/FunctionalTesting/Private/ScreenshotFunctionalTestBase.cpp
+++ b/Engine/Source/Developer/FunctionalTesting/Private/ScreenshotFunctionalTestBase.cpp
@@ -19,7 +19,6 @@
 #include "Logging/LogMacros.h"
 #include "UObject/AutomationObjectVersion.h"
 
-#define	WITH_AUTOMATION_TESTS			(WITH_DEV_AUTOMATION_TESTS || WITH_PERF_AUTOMATION_TESTS)
 #define	WITH_EDITOR_AUTOMATION_TESTS	(WITH_EDITOR && WITH_AUTOMATION_TESTS)
 
 DEFINE_LOG_CATEGORY_STATIC(LogScreenshotFunctionalTest, Log, Log)
diff --git a/Engine/Source/Developer/FunctionalTesting/Public/AutomationBlueprintFunctionLibrary.h b/Engine/Source/Developer/FunctionalTesting/Public/AutomationBlueprintFunctionLibrary.h
--- a/Engine/Source/Developer/FunctionalTesting/Public/AutomationBlueprintFunctionLibrary.h
+++ b/Engine/Source/Developer/FunctionalTesting/Public/AutomationBlueprintFunctionLibrary.h
@@ -8,7 +8,7 @@
 #include "AutomationScreenshotOptions.h"
 #include "HAL/IConsoleManager.h"
 #include "Templates/UniquePtr.h"
-
+#include "Misc/AutomationTest.h"
 #include "AutomationBlueprintFunctionLibrary.generated.h"
 
 class ACameraActor;
@@ -156,7 +156,7 @@ public:
 	static void SetScalabilityQualityToLow(UObject* WorldContextObject);
 };
 
-#if (WITH_DEV_AUTOMATION_TESTS || WITH_PERF_AUTOMATION_TESTS)
+#if WITH_AUTOMATION_TESTS
 
 template<typename T>
 class FConsoleVariableSwapperTempl
diff --git a/Engine/Source/Runtime/AutomationWorker/Private/AutomationWorkerModule.cpp b/Engine/Source/Runtime/AutomationWorker/Private/AutomationWorkerModule.cpp
--- a/Engine/Source/Runtime/AutomationWorker/Private/AutomationWorkerModule.cpp
+++ b/Engine/Source/Runtime/AutomationWorker/Private/AutomationWorkerModule.cpp
@@ -414,7 +414,7 @@ void FAutomationWorkerModule::HandlePerformanceDataRetrieved(const FAutomationWo
 #if WITH_ENGINE
 void FAutomationWorkerModule::HandleScreenShotCapturedWithName(const TArray<FColor>& RawImageData, const FAutomationScreenshotData& Data)
 {
-#if (WITH_DEV_AUTOMATION_TESTS || WITH_PERF_AUTOMATION_TESTS)
+#if WITH_AUTOMATION_TESTS
 	int32 NewHeight = Data.Height;
 	int32 NewWidth = Data.Width;
 
diff --git a/Engine/Source/Runtime/Core/Public/Misc/AutomationTest.h b/Engine/Source/Runtime/Core/Public/Misc/AutomationTest.h
--- a/Engine/Source/Runtime/Core/Public/Misc/AutomationTest.h
+++ b/Engine/Source/Runtime/Core/Public/Misc/AutomationTest.h
@@ -28,6 +28,11 @@
 #include "Misc/AutomationEvent.h"
 #include "Internationalization/Regex.h"
 
+
+#ifndef WITH_AUTOMATION_TESTS
+	#define WITH_AUTOMATION_TESTS (WITH_DEV_AUTOMATION_TESTS || WITH_PERF_AUTOMATION_TESTS)
+#endif
+
 /** Flags for specifying automation test requirements/behavior */
 namespace EAutomationTestFlags
 {
diff --git a/Engine/Source/Runtime/Engine/Private/Tests/AutomationCommon.cpp b/Engine/Source/Runtime/Engine/Private/Tests/AutomationCommon.cpp
--- a/Engine/Source/Runtime/Engine/Private/Tests/AutomationCommon.cpp
+++ b/Engine/Source/Runtime/Engine/Private/Tests/AutomationCommon.cpp
@@ -22,7 +22,7 @@
 #include "StereoRendering.h"
 #include "Misc/PackageName.h"
 
-#if (WITH_DEV_AUTOMATION_TESTS || WITH_PERF_AUTOMATION_TESTS)
+#if WITH_AUTOMATION_TESTS
 
 DEFINE_LOG_CATEGORY_STATIC(LogEngineAutomationLatentCommand, Log, All);
 DEFINE_LOG_CATEGORY(LogEditorAutomationTests);
@@ -79,7 +79,7 @@ namespace AutomationCommon
 		return HardwareDetailsString;
 	}
 
-#if (WITH_DEV_AUTOMATION_TESTS || WITH_PERF_AUTOMATION_TESTS)
+#if WITH_AUTOMATION_TESTS
 
 	/** Gets a path used for automation testing (PNG sent to the AutomationTest folder) */
 	FString GetScreenshotName(const FString& TestName)
diff --git a/Engine/Source/Runtime/Engine/Public/Tests/AutomationCommon.h b/Engine/Source/Runtime/Engine/Public/Tests/AutomationCommon.h
--- a/Engine/Source/Runtime/Engine/Public/Tests/AutomationCommon.h
+++ b/Engine/Source/Runtime/Engine/Public/Tests/AutomationCommon.h
@@ -8,7 +8,7 @@
 class AMatineeActor;
 class SWindow;
 
-#if (WITH_DEV_AUTOMATION_TESTS || WITH_PERF_AUTOMATION_TESTS)
+#if WITH_AUTOMATION_TESTS
 
 ///////////////////////////////////////////////////////////////////////
 // Common Latent commands which are used across test type. I.e. Engine, Network, etc...
@@ -26,7 +26,7 @@ namespace AutomationCommon
 	/** Get a string contains the render mode we are currently in */
 	ENGINE_API FString GetRenderDetailsString();
 
-#if (WITH_DEV_AUTOMATION_TESTS || WITH_PERF_AUTOMATION_TESTS)
+#if WITH_AUTOMATION_TESTS
 
 	/** Gets a name to be used for this screenshot.  This will return something like 
 		TestName/PlatformName/DeviceName.png. It's important to understand that a screenshot
@@ -52,7 +52,7 @@ namespace AutomationCommon
 #endif
 }
 
-#if (WITH_DEV_AUTOMATION_TESTS || WITH_PERF_AUTOMATION_TESTS)
+#if WITH_AUTOMATION_TESTS
 
 /**
  * Parameters to the Latent Automation command FTakeEditorScreenshotCommand
