Fix issue with functional test base ignorning error/warning logging settings (GH-6641)

diff --git a/Engine/Source/Developer/FunctionalTesting/Classes/FunctionalTestBase.h b/Engine/Source/Developer/FunctionalTesting/Classes/FunctionalTestBase.h
--- a/Engine/Source/Developer/FunctionalTesting/Classes/FunctionalTestBase.h
+++ b/Engine/Source/Developer/FunctionalTesting/Classes/FunctionalTestBase.h
@@ -30,8 +30,17 @@ public:
 	void SetLogErrorAndWarningHandling(TOptional<bool> LogErrorsAreErrors, TOptional<bool> LogWarningsAreErrors)
 	{
 		SetLogErrorAndWarningHandlingToDefault();
-	}
-	
+
+		if (LogErrorsAreErrors.IsSet())
+		{
+			bTreatLogErrorsAsErrors = LogErrorsAreErrors.GetValue();
+		}
+
+		if (LogWarningsAreErrors.IsSet())
+		{
+			bTreatLogWarningsAsErrors = LogWarningsAreErrors.GetValue();
+		}
+	}	
 
 	/**
 	 * Determines if Error logs should be considered test errors
