- Fix for case where people try to take an automation screenshot during BeginPlay.

Now perform the pre-screenshot flush of loading on the first tick of the latent action and not in its constructor. This addresses cases where BP's attempt to take a screenshot in their BeginPlay event while the world is being created.

- Fix for crash that can happen when a window needs resized after a screenshot is captured.

Don't delete (which leads to restoring the viewport size) on response to the screenshot delegates. Instead queue it by one frame

diff --git a/Engine/Source/Developer/FunctionalTesting/Private/AutomationBlueprintFunctionLibrary.cpp b/Engine/Source/Developer/FunctionalTesting/Private/AutomationBlueprintFunctionLibrary.cpp
--- a/Engine/Source/Developer/FunctionalTesting/Private/AutomationBlueprintFunctionLibrary.cpp
+++ b/Engine/Source/Developer/FunctionalTesting/Private/AutomationBlueprintFunctionLibrary.cpp
@@ -44,11 +44,13 @@
 #include "Engine/LevelStreaming.h"
 #include "Templates/UnrealTemplate.h"
 #include "UObject/GCObjectScopeGuard.h"
+#include "Containers/Ticker.h"
 
 #if WITH_EDITOR
 #include "SLevelViewport.h"
 #endif
 
+
 #define LOCTEXT_NAMESPACE "Automation"
 
 DEFINE_LOG_CATEGORY_STATIC(BlueprintAssertion, Error, Error)
@@ -357,6 +359,7 @@ public:
 		, Notes(InNotes)
 		, Options(InOptions)
 		, bNeedsViewportSizeRestore(false)
+		, bDeleteQueued(false)
 	{
 		EnvSetup.Setup(InWorld, Options);
 
@@ -407,9 +410,9 @@ public:
 		{
 			if (GEngine->GameViewport)
 			{
-			FSceneViewport* GameViewport = GEngine->GameViewport->GetGameViewport();
-			GameViewport->SetViewportSize(ViewportRestoreSize.X, ViewportRestoreSize.Y);
-		}
+				FSceneViewport* GameViewport = GEngine->GameViewport->GetGameViewport();
+				GameViewport->SetViewportSize(ViewportRestoreSize.X, ViewportRestoreSize.Y);
+			}
 		}
 
 		EnvSetup.Restore();
@@ -417,6 +420,18 @@ public:
 		FAutomationTestFramework::Get().NotifyScreenshotTakenAndCompared();
 	}
 
+	void DeleteSelfNextFrame()
+	{
+		if (!bDeleteQueued)
+		{
+			FTicker::GetCoreTicker().AddTicker(TEXT("ScreenshotCleanup"), 0.1, [this](float) {
+				delete this;
+				return false;
+				});
+			bDeleteQueued = true;
+		}
+	}
+
 	void GrabScreenShot(int32 InSizeX, int32 InSizeY, const TArray<FColor>& InImageData)
 	{
 		check(IsInGameThread());
@@ -445,24 +460,23 @@ public:
 
 			UE_LOG(AutomationFunctionLibrary, Log, TEXT("Screenshot captured as %s"), *Data.ScreenshotName);
 
-			if ( GIsAutomationTesting )
+			if (GIsAutomationTesting)
 			{
 				FAutomationTestFramework::Get().OnScreenshotCompared.AddRaw(this, &FAutomationScreenshotTaker::OnComparisonComplete);
 				FScreenshotRequest::OnScreenshotRequestProcessed().RemoveAll(this);
 				return;
 			}
 		}
-		
-		delete this;
-		}
+
+		DeleteSelfNextFrame();
+	}
 
 	void OnScreenshotProcessed()
-		{
+	{
 		UE_LOG(AutomationFunctionLibrary, Log, TEXT("Screenshot processed, but not compared."));
 
-		// If it's done being processed 
-			delete this;
-		}
+		DeleteSelfNextFrame();
+	}
 
 	void OnComparisonComplete(const FAutomationScreenshotCompareResults& CompareResults)
 	{
@@ -473,19 +487,19 @@ public:
 			CurrentTest->AddEvent(CompareResults.ToAutomationEvent(Name));
 		}
 
-		delete this;
-			}
+		DeleteSelfNextFrame();
+	}
 
 	void WorldDestroyed(ULevel* InLevel, UWorld* InWorld)
-			{
+	{
 		// If the InLevel is null, it's a signal that the entire world is about to disappear, so
 		// go ahead and remove this widget from the viewport, it could be holding onto too many
 		// dangerous actor references that won't carry over into the next world.
 		if (InLevel == nullptr && InWorld == World.Get())
-				{
+		{
 			delete this;
-				}
-			}
+		}
+	}
 
 private:
 
@@ -498,6 +512,7 @@ private:
 	FAutomationTestScreenshotEnvSetup EnvSetup;
 	FIntPoint ViewportRestoreSize;
 	bool bNeedsViewportSizeRestore;
+	bool bDeleteQueued;
 };
 
 class FAutomationHighResScreenshotGrabber
diff --git a/Engine/Source/Developer/FunctionalTesting/Private/TakeScreenshotAfterTimeLatentAction.cpp b/Engine/Source/Developer/FunctionalTesting/Private/TakeScreenshotAfterTimeLatentAction.cpp
--- a/Engine/Source/Developer/FunctionalTesting/Private/TakeScreenshotAfterTimeLatentAction.cpp
+++ b/Engine/Source/Developer/FunctionalTesting/Private/TakeScreenshotAfterTimeLatentAction.cpp
@@ -14,11 +14,12 @@ FTakeScreenshotAfterTimeLatentAction::FTakeScreenshotAfterTimeLatentAction(const
 	, ScreenshotName(InScreenshotName)
 	, Notes(InNotes)
 	, SecondsRemaining(InOptions.Delay)
+	, FinishedLoading(false)
 	, IssuedScreenshotCapture(false)
 	, TakenScreenshot(false)
 	, Options(InOptions)
 {
-	UAutomationBlueprintFunctionLibrary::FinishLoadingBeforeScreenshot();
+
 }
 
 FTakeScreenshotAfterTimeLatentAction::~FTakeScreenshotAfterTimeLatentAction()
@@ -33,6 +34,12 @@ void FTakeScreenshotAfterTimeLatentAction::OnScreenshotTakenAndCompared()
 
 void FTakeScreenshotAfterTimeLatentAction::UpdateOperation(FLatentResponse& Response)
 {
+	if (!FinishedLoading)
+	{
+		UAutomationBlueprintFunctionLibrary::FinishLoadingBeforeScreenshot();
+		FinishedLoading = true;
+	}
+
 	if ( !TakenScreenshot )
 	{
 		if ( !IssuedScreenshotCapture )
diff --git a/Engine/Source/Developer/FunctionalTesting/Public/TakeScreenshotAfterTimeLatentAction.h b/Engine/Source/Developer/FunctionalTesting/Public/TakeScreenshotAfterTimeLatentAction.h
--- a/Engine/Source/Developer/FunctionalTesting/Public/TakeScreenshotAfterTimeLatentAction.h
+++ b/Engine/Source/Developer/FunctionalTesting/Public/TakeScreenshotAfterTimeLatentAction.h
@@ -31,6 +31,7 @@ private:
 	FString ScreenshotName;
 	FString Notes;
 	float SecondsRemaining;
+	bool FinishedLoading;
 	bool IssuedScreenshotCapture;
 	bool TakenScreenshot;
 	FAutomationScreenshotOptions Options;
