Add asserts compatabile with doubles for Functional Tests

#rnx

diff --git a/Engine/Source/Developer/FunctionalTesting/Classes/FunctionalTest.h b/Engine/Source/Developer/FunctionalTesting/Classes/FunctionalTest.h
--- a/Engine/Source/Developer/FunctionalTesting/Classes/FunctionalTest.h
+++ b/Engine/Source/Developer/FunctionalTesting/Classes/FunctionalTest.h
@@ -371,6 +371,13 @@ public:
 	UFUNCTION(BlueprintCallable, Category = "Asserts", DisplayName = "Assert Value (Float)", meta = ( HidePin = "ContextObject", DefaultToSelf = "ContextObject"))
 	bool AssertValue_Float(float Actual, EComparisonMethod ShouldBe, float Expected, const FString& What, const UObject* ContextObject = nullptr);
 
+	/**
+	 * Assert on a relationship between two doubles.
+	 * @param What	A name to use in the message if the assert fails (What: expected {Actual} to be <ShouldBe> {Expected} for context '')
+	 */
+	UFUNCTION(BlueprintCallable, Category = "Asserts", DisplayName = "Assert Value (Double)", meta = (HidePin = "ContextObject", DefaultToSelf = "ContextObject"))
+	bool AssertValue_Double(double Actual, EComparisonMethod ShouldBe, double Expected, const FString& What, const UObject* ContextObject = nullptr);
+
 	/**
 	 * Assert on a relationship between two DateTimes.
 	 * @param What	A name to use in the message if the assert fails (What: expected {Actual} to be <ShouldBe> {Expected} for context '')
@@ -392,6 +399,13 @@ public:
 	UFUNCTION(BlueprintCallable, Category = "Asserts", DisplayName = "Assert Equal (Float)", meta = ( HidePin = "ContextObject", DefaultToSelf = "ContextObject"))
 	bool AssertEqual_Float(const float Actual, const float Expected, const FString& What, const float Tolerance = 1.e-4, const UObject* ContextObject = nullptr);
 
+	/**
+	 * Assert that two double are equal within tolerance between two doubles.
+	 * @param What	A name to use in the message if the assert fails (What: expected {Actual} to be Equal To {Expected} within Tolerance for context '')
+	 */
+	UFUNCTION(BlueprintCallable, Category = "Asserts", DisplayName = "Assert Equal (Double)", meta = (HidePin = "ContextObject", DefaultToSelf = "ContextObject"))
+	bool AssertEqual_Double(const double Actual, const double Expected, const FString& What, const double Tolerance = 1.e-4, const UObject* ContextObject = nullptr);
+
 	/**
 	* Assert that two bools are equal
 	* @param What	A name to use in the message if the assert fails (What: expected {Actual} to be Equal To {Expected} for context '')
diff --git a/Engine/Source/Developer/FunctionalTesting/Private/FunctionalTest.cpp b/Engine/Source/Developer/FunctionalTesting/Private/FunctionalTest.cpp
--- a/Engine/Source/Developer/FunctionalTesting/Private/FunctionalTest.cpp
+++ b/Engine/Source/Developer/FunctionalTesting/Private/FunctionalTest.cpp
@@ -648,6 +648,20 @@ bool AFunctionalTest::AssertValue_Float(float Actual, EComparisonMethod ShouldBe
 	}
 }
 
+bool AFunctionalTest::AssertValue_Double(double Actual, EComparisonMethod ShouldBe, double Expected, const FString& What, const UObject* ContextObject)
+{
+	if ( !PerformComparison(Actual, Expected, ShouldBe) )
+	{
+		LogStep(ELogVerbosity::Error, FString::Printf(TEXT("%s: expected {%lf} to be %s {%lf} for context '%s'"), *What, Actual, *GetComparisonAsString(ShouldBe), Expected, ContextObject ? *ContextObject->GetName() : TEXT("")));
+		return false;
+	}
+	else
+	{
+		LogStep(ELogVerbosity::Log, FString::Printf(TEXT("%s: expected {%lf} to be %s {%lf} for context '%s'"), *What, Actual, *GetComparisonAsString(ShouldBe), Expected, ContextObject ? *ContextObject->GetName() : TEXT("")));
+		return true;
+	}
+}
+
 bool AFunctionalTest::AssertValue_DateTime(FDateTime Actual, EComparisonMethod ShouldBe, FDateTime Expected, const FString& What, const UObject* ContextObject)
 {
 	if ( !PerformComparison(Actual, Expected, ShouldBe) )
@@ -676,6 +690,20 @@ bool AFunctionalTest::AssertEqual_Float(const float Actual, const float Expected
 	}
 }
 
+bool AFunctionalTest::AssertEqual_Double(const double Actual, const double Expected, const FString& What, const double Tolerance, const UObject* ContextObject)
+{
+	if ( !FMath::IsNearlyEqual(Actual, Expected, Tolerance) )
+	{
+		LogStep(ELogVerbosity::Error, FString::Printf(TEXT("Expected '%s' to be {%lf}, but it was {%lf} within tolerance {%lf} for context '%s'"), *What, Expected, Actual, Tolerance, ContextObject ? *ContextObject->GetName() : TEXT("")));
+		return false;
+	}
+	else
+	{
+		LogStep(ELogVerbosity::Log, FString::Printf(TEXT("Double assertion passed (%s)"), *What));
+		return true;
+	}
+}
+
 bool AFunctionalTest::AssertEqual_Bool(const bool Actual, const bool Expected, const FString& What, const UObject* ContextObject)
 {
 	if (Actual != Expected)
