Expose FAutomationTestBase::AddTelemetryData to blueprint and python

#rnx

diff --git a/Engine/Source/Developer/FunctionalTesting/Private/AutomationBlueprintFunctionLibrary.cpp b/Engine/Source/Developer/FunctionalTesting/Private/AutomationBlueprintFunctionLibrary.cpp
--- a/Engine/Source/Developer/FunctionalTesting/Private/AutomationBlueprintFunctionLibrary.cpp
+++ b/Engine/Source/Developer/FunctionalTesting/Private/AutomationBlueprintFunctionLibrary.cpp
@@ -1216,6 +1216,28 @@ UAutomationEditorTask* UAutomationBlueprintFunctionLibrary::TakeHighResScreensho
 	return Task;
 }
 
+void UAutomationBlueprintFunctionLibrary::AddTestTelemetryData(FString DataPoint, float Measurement, FString Context)
+{
+	if (GIsAutomationTesting)
+	{
+		if (FAutomationTestBase* CurrentTest = FAutomationTestFramework::Get().GetCurrentTest())
+		{
+			CurrentTest->AddTelemetryData(DataPoint, Measurement, Context);
+		}
+	}
+}
+
+void UAutomationBlueprintFunctionLibrary::SetTestTelemetryStorage(FString StorageName)
+{
+	if (GIsAutomationTesting)
+	{
+		if (FAutomationTestBase* CurrentTest = FAutomationTestFramework::Get().GetCurrentTest())
+		{
+			CurrentTest->SetTelemetryStorage(StorageName);
+		}
+	}
+}
+
 FAutomationScreenshotOptions UAutomationBlueprintFunctionLibrary::GetDefaultScreenshotOptionsForGameplay(EComparisonTolerance Tolerance, float Delay)
 {
 	FAutomationScreenshotOptions Options;
diff --git a/Engine/Source/Developer/FunctionalTesting/Public/AutomationBlueprintFunctionLibrary.h b/Engine/Source/Developer/FunctionalTesting/Public/AutomationBlueprintFunctionLibrary.h
--- a/Engine/Source/Developer/FunctionalTesting/Public/AutomationBlueprintFunctionLibrary.h
+++ b/Engine/Source/Developer/FunctionalTesting/Public/AutomationBlueprintFunctionLibrary.h
@@ -137,6 +137,18 @@ public:
 	UFUNCTION(BlueprintCallable, Category = "Automation", meta = (AdvancedDisplay="Camera, bMaskEnabled, bCaptureHDR, ComparisonTolerance, ComparisonNotes"))
 	static UAutomationEditorTask* TakeHighResScreenshot(int32 ResX, int32 ResY, FString Filename, ACameraActor* Camera = nullptr, bool bMaskEnabled = false, bool bCaptureHDR = false, EComparisonTolerance ComparisonTolerance = EComparisonTolerance::Low, FString ComparisonNotes = TEXT(""), float Delay = 0.0);
 
+	/**
+	* Add Telemetry data to currently running automated test.
+	*/
+	UFUNCTION(BlueprintCallable, Category = "Automation", meta = (AdvancedDisplay = "Context"))
+	static void AddTestTelemetryData(FString DataPoint, float Measurement, FString Context = TEXT(""));
+
+	/**
+	* Set Telemetry data storage name of currently running automated test.
+	*/
+	UFUNCTION(BlueprintCallable, Category = "Automation")
+	static void SetTestTelemetryStorage(FString StorageName);
+
 	/**
 	 * 
 	 */
