Removal of UE4 reference in comment

diff --git a/Engine/Source/Developer/FunctionalTesting/Private/ScreenshotFunctionalTest.cpp b/Engine/Source/Developer/FunctionalTesting/Private/ScreenshotFunctionalTest.cpp
--- a/Engine/Source/Developer/FunctionalTesting/Private/ScreenshotFunctionalTest.cpp
+++ b/Engine/Source/Developer/FunctionalTesting/Private/ScreenshotFunctionalTest.cpp
@@ -76,7 +76,7 @@ void AScreenshotFunctionalTest::RequestScreenshot()
 	}
 	else
 	{
-		// Screenshots in UE4 work in this way:
+		// Screenshots in Unreal Engine work in this way:
 		// 1. Call FScreenshotRequest::RequestScreenshot to ask the system to take a screenshot. The screenshot
 		//    will have the same resolution as the current viewport;
 		// 2. Register a callback to UGameViewportClient::OnScreenshotCaptured() delegate. The call back will be
