Fixed AFunctionalAITestBase::IsReady_Implementation to pass when navmesh is not ready, but the test is explicitly configured to not wait for navmesh.

This supports pathfollowing-less tests on navmeshless maps.

diff --git a/Engine/Source/Developer/FunctionalTesting/Private/FunctionalAITest.cpp b/Engine/Source/Developer/FunctionalTesting/Private/FunctionalAITest.cpp
--- a/Engine/Source/Developer/FunctionalTesting/Private/FunctionalAITest.cpp
+++ b/Engine/Source/Developer/FunctionalTesting/Private/FunctionalAITest.cpp
@@ -86,7 +86,7 @@ void AFunctionalAITestBase::StartTest()
 
 bool AFunctionalAITestBase::IsReady_Implementation()
 {
-	return Super::IsReady_Implementation() && IsNavMeshReady();
+	return Super::IsReady_Implementation() && (bWaitForNavMesh == false || IsNavMeshReady());
 }
 
 void AFunctionalAITestBase::OnTimeout()
