Use GetActorLabel instead of GetName when taking Screenshot with functional test

Actor Name is a property that user can't change directly from the Outliner now. To keep to ScreenshotFunctionalTest actor behavior seamless for the user, we use the actor label instead.
A property is introduced to hold the Actor name on the Client build.

+ update all the necessary screenshot names.

#rnx

diff --git a/Engine/Source/Developer/FunctionalTesting/Classes/FunctionalTest.h b/Engine/Source/Developer/FunctionalTesting/Classes/FunctionalTest.h
--- a/Engine/Source/Developer/FunctionalTesting/Classes/FunctionalTest.h
+++ b/Engine/Source/Developer/FunctionalTesting/Classes/FunctionalTest.h
@@ -224,6 +224,9 @@ class FUNCTIONALTESTING_API AFunctionalTest : public AActor
 public:
 	AFunctionalTest(const FObjectInitializer& ObjectInitializer = FObjectInitializer::Get());
 
+	UPROPERTY(BlueprintReadOnly, Category = "Functional Testing")
+	FString TestLabel;
+
 private:
 	UPROPERTY()
 	TObjectPtr<UBillboardComponent> SpriteComponent;
diff --git a/Engine/Source/Developer/FunctionalTesting/Private/FunctionalTest.cpp b/Engine/Source/Developer/FunctionalTesting/Private/FunctionalTest.cpp
--- a/Engine/Source/Developer/FunctionalTesting/Private/FunctionalTest.cpp
+++ b/Engine/Source/Developer/FunctionalTesting/Private/FunctionalTest.cpp
@@ -87,6 +87,7 @@ namespace
 
 AFunctionalTest::AFunctionalTest( const FObjectInitializer& ObjectInitializer )
 	: Super(ObjectInitializer)
+	, TestLabel(GetName())
 	, bIsEnabled(true)
 	, LogErrorHandling(EFunctionalTestLogHandling::ProjectDefault)
 	, LogWarningHandling(EFunctionalTestLogHandling::ProjectDefault)
@@ -173,6 +174,8 @@ void AFunctionalTest::OnConstruction(const FTransform& Transform)
 	Super::OnConstruction(Transform);
 
 #if WITH_EDITOR
+	TestLabel = GetActorLabel();
+
 	if ( TestName )
 	{
 		if ( bIsEnabled )
@@ -217,7 +220,7 @@ bool AFunctionalTest::RunTest(const TArray<FString>& Params)
 	if (FunctionalTest)
 	{
 		FunctionalTest->SetLogErrorAndWarningHandling(bSuppressErrors, bSuppressWarnings, bWarningsAreErrors);
-		FunctionalTest->SetFunctionalTestRunning(GetName());
+		FunctionalTest->SetFunctionalTestRunning(TestLabel);
 	}
 
 	FailureMessage = TEXT("");
@@ -373,7 +376,7 @@ void AFunctionalTest::FinishTest(EFunctionalTestResult TestResult, const FString
 	FFunctionalTestBase* FunctionalTest = static_cast<FFunctionalTestBase*>(FAutomationTestFramework::Get().GetCurrentTest());
 	if (FunctionalTest)
 	{
-		FunctionalTest->SetFunctionalTestComplete(GetName());
+		FunctionalTest->SetFunctionalTestComplete(TestLabel);
 	}
 
 	bIsRunning = false;
diff --git a/Engine/Source/Developer/FunctionalTesting/Private/ScreenshotFunctionalTestBase.cpp b/Engine/Source/Developer/FunctionalTesting/Private/ScreenshotFunctionalTestBase.cpp
--- a/Engine/Source/Developer/FunctionalTesting/Private/ScreenshotFunctionalTestBase.cpp
+++ b/Engine/Source/Developer/FunctionalTesting/Private/ScreenshotFunctionalTestBase.cpp
@@ -125,9 +125,9 @@ void AScreenshotFunctionalTestBase::OnScreenShotCaptured(int32 InSizeX, int32 In
 	GEngine->GameViewport->OnScreenshotCaptured().RemoveAll(this);
 
 #if WITH_AUTOMATION_TESTS
-	TArray<uint8> CapturedFrameTrace = AutomationCommon::CaptureFrameTrace(GetWorld()->GetName(), GetName());
+	TArray<uint8> CapturedFrameTrace = AutomationCommon::CaptureFrameTrace(GetWorld()->GetName(), TestLabel);
 
-	FAutomationScreenshotData Data = UAutomationBlueprintFunctionLibrary::BuildScreenshotData(GetWorld()->GetName(), GetName(), InSizeX, InSizeY);
+	FAutomationScreenshotData Data = UAutomationBlueprintFunctionLibrary::BuildScreenshotData(GetWorld()->GetName(), TestLabel, InSizeX, InSizeY);
 
 	// Copy the relevant data into the metadata for the screenshot.
 	Data.bHasComparisonRules = true;
@@ -174,7 +174,7 @@ void AScreenshotFunctionalTestBase::OnComparisonComplete(const FAutomationScreen
 
 	if (FAutomationTestBase* CurrentTest = FAutomationTestFramework::Get().GetCurrentTest())
 	{
-		CurrentTest->AddEvent(CompareResults.ToAutomationEvent(GetName()));
+		CurrentTest->AddEvent(CompareResults.ToAutomationEvent(TestLabel));
 	}
 
 	FAutomationTestFramework::Get().NotifyScreenshotTakenAndCompared();
