Add CallInEditor function specifier to DebugGatherRelevantActors in FunctionalTest.h
#8054

diff --git a/Engine/Source/Developer/FunctionalTesting/Classes/FunctionalTest.h b/Engine/Source/Developer/FunctionalTesting/Classes/FunctionalTest.h
--- a/Engine/Source/Developer/FunctionalTesting/Classes/FunctionalTest.h
+++ b/Engine/Source/Developer/FunctionalTesting/Classes/FunctionalTest.h
@@ -524,7 +524,7 @@ public:
 public:
 
 	/** Used by debug drawing to gather actors this test is using and point at them on the level to better understand test's setup */
-	UFUNCTION(BlueprintImplementableEvent, Category="Functional Testing")
+	UFUNCTION(BlueprintImplementableEvent, CallInEditor, Category="Functional Testing")
 	TArray<AActor*> DebugGatherRelevantActors() const;
 
 	virtual void GatherRelevantActors(TArray<AActor*>& OutActors) const;
