Hide the debug messages for FunctionalTest.

diff --git a/Engine/Source/Developer/FunctionalTesting/Classes/FunctionalUIScreenshotTest.h b/Engine/Source/Developer/FunctionalTesting/Classes/FunctionalUIScreenshotTest.h
--- a/Engine/Source/Developer/FunctionalTesting/Classes/FunctionalUIScreenshotTest.h
+++ b/Engine/Source/Developer/FunctionalTesting/Classes/FunctionalUIScreenshotTest.h
@@ -30,6 +30,9 @@ class AFunctionalUIScreenshotTest : public AScreenshotFunctionalTestBase
 public:
 	AFunctionalUIScreenshotTest(const FObjectInitializer& ObjectInitializer);
 
+public:
+	virtual void EndPlay(const EEndPlayReason::Type EndPlayReason) override;
+
 protected:
 	virtual void PrepareTest() override;
 
@@ -38,7 +41,6 @@ protected:
 	virtual void RequestScreenshot() override;
 	
 protected:
-
 	UPROPERTY(EditAnywhere, Category = "UI")
 	TSubclassOf<UUserWidget> WidgetClass;
 
@@ -50,4 +52,10 @@ protected:
 
 	UPROPERTY(Transient, DuplicateTransient)
 	TObjectPtr<UTextureRenderTarget2D> ScreenshotRT;
+
+	UPROPERTY(EditAnywhere, Category = "UI")
+	bool bHideDebugCanvas;
+
+private:
+	TOptional<bool> PreviousDebugCanvasVisible;
 };
diff --git a/Engine/Source/Developer/FunctionalTesting/Private/FunctionalUIScreenshotTest.cpp b/Engine/Source/Developer/FunctionalTesting/Private/FunctionalUIScreenshotTest.cpp
--- a/Engine/Source/Developer/FunctionalTesting/Private/FunctionalUIScreenshotTest.cpp
+++ b/Engine/Source/Developer/FunctionalTesting/Private/FunctionalUIScreenshotTest.cpp
@@ -19,6 +19,7 @@ AFunctionalUIScreenshotTest::AFunctionalUIScreenshotTest( const FObjectInitializ
 	: AScreenshotFunctionalTestBase(ObjectInitializer)
 {
 	WidgetLocation = EWidgetTestAppearLocation::Viewport;
+	bHideDebugCanvas = true;
 }
 
 /**
@@ -53,6 +54,16 @@ void AFunctionalUIScreenshotTest::PrepareTest()
 	// Resize viewport to screenshot size
 	Super::PrepareTest();
 
+	// Hide all debug info
+	if (bHideDebugCanvas)
+	{
+		if (IConsoleVariable* CVarDebugCanvasVisible = IConsoleManager::Get().FindConsoleVariable(TEXT("Slate.GameLayer.DebugCanvasVisible")))
+		{
+			PreviousDebugCanvasVisible = CVarDebugCanvasVisible->GetBool();
+			CVarDebugCanvasVisible->Set(false);
+		}
+	}
+
 	TSharedPtr<SViewport> GameViewportWidget = GEngine->GameViewport->GetGameViewportWidget();
 	check(GameViewportWidget.IsValid());
 
@@ -107,6 +118,31 @@ void AFunctionalUIScreenshotTest::OnScreenshotTakenAndCompared()
 
 	// Restore viewport size and finish the test
 	Super::OnScreenshotTakenAndCompared();
+
+	// Restore the debug text
+	if (PreviousDebugCanvasVisible.IsSet())
+	{
+		if (IConsoleVariable* CVarDebugCanvasVisible = IConsoleManager::Get().FindConsoleVariable(TEXT("Slate.GameLayer.DebugCanvasVisible")))
+		{
+			CVarDebugCanvasVisible->Set(PreviousDebugCanvasVisible.GetValue());
+			PreviousDebugCanvasVisible.Reset();
+		}
+	}
+}
+
+void AFunctionalUIScreenshotTest::EndPlay(const EEndPlayReason::Type EndPlayReason)
+{
+	Super::EndPlay(EndPlayReason);
+
+	// Restore the debug text
+	if (PreviousDebugCanvasVisible.IsSet())
+	{
+		if (IConsoleVariable* CVarDebugCanvasVisible = IConsoleManager::Get().FindConsoleVariable(TEXT("Slate.GameLayer.DebugCanvasVisible")))
+		{
+			CVarDebugCanvasVisible->Set(PreviousDebugCanvasVisible.GetValue());
+			PreviousDebugCanvasVisible.Reset();
+		}
+	}
 }
 
 /**
