Fix for double delete with FAutomationScreenshotTaker

diff --git a/Engine/Source/Developer/FunctionalTesting/Private/AutomationBlueprintFunctionLibrary.cpp b/Engine/Source/Developer/FunctionalTesting/Private/AutomationBlueprintFunctionLibrary.cpp
--- a/Engine/Source/Developer/FunctionalTesting/Private/AutomationBlueprintFunctionLibrary.cpp
+++ b/Engine/Source/Developer/FunctionalTesting/Private/AutomationBlueprintFunctionLibrary.cpp
@@ -479,7 +479,10 @@ public:
 		// dangerous actor references that won't carry over into the next world.
 		if (InLevel == nullptr && InWorld == World.Get())
 		{
-			delete this;
+			// we don't delete directly because of the risk of conflicting with an already in flight
+			// request to delete ourselves
+			World.Reset();
+			DeleteSelfNextFrame();
 		}
 	}
 
