USD: Improve reproducibility of automated tests.

Implement new function for switching view mode and visualize buffer before taking screenshots;
Screenshot individual buffers instead of using lighting on screenshot tests;
Fix import tests using whatever render context to import was serialized into the config for the import tests;
Disable geometry cache notification on other geometry cache tests;
Increase tolerance on test_pie_streaming_level to give more time for levels to stream in and out before checking for loaded actors;

diff --git a/Engine/Source/Developer/FunctionalTesting/Private/AutomationBlueprintFunctionLibrary.cpp b/Engine/Source/Developer/FunctionalTesting/Private/AutomationBlueprintFunctionLibrary.cpp
--- a/Engine/Source/Developer/FunctionalTesting/Private/AutomationBlueprintFunctionLibrary.cpp
+++ b/Engine/Source/Developer/FunctionalTesting/Private/AutomationBlueprintFunctionLibrary.cpp
@@ -158,7 +158,7 @@ public:
 		, Options(InOptions)
 	{
 	}
-	
+
 	/** ISceneViewExtension interface */
 	virtual void SetupView(FSceneViewFamily& InViewFamily, FSceneView& InView)
 	{
@@ -216,7 +216,7 @@ public:
 			// Disable screen percentage.
 			//InViewFamily.EngineShowFlags.SetScreenPercentage(false);
 		}
-		
+
 		if (Options.bDisableTonemapping)
 		{
 			//InViewFamily.EngineShowFlags.SetEyeAdaptation(false);
@@ -370,13 +370,13 @@ public:
 			{
 #if WITH_EDITOR
 				// In the editor we can only attempt to re-size standalone viewports
-				UEditorEngine* EditorEngine = Cast<UEditorEngine>(GEngine);	
+				UEditorEngine* EditorEngine = Cast<UEditorEngine>(GEngine);
 
-				const bool bIsPIEViewport = GameViewport->IsPlayInEditorViewport();	
+				const bool bIsPIEViewport = GameViewport->IsPlayInEditorViewport();
 				const bool bIsNewViewport = InWorld && EditorEngine && EditorEngine->WorldIsPIEInNewViewport(InWorld);
 
 				if (!bIsPIEViewport || bIsNewViewport)
-#endif		
+#endif
 				{
 					ViewportRestoreSize = GameViewport->GetSize();
 					FIntPoint ScreenshotViewportSize = UAutomationBlueprintFunctionLibrary::GetAutomationScreenshotSize(InOptions);
@@ -507,7 +507,7 @@ public:
 private:
 
 	TWeakObjectPtr<UWorld> World;
-	
+
 	FString	Context;
 	FString	ScreenShotName;
 	FString Notes;
@@ -636,7 +636,7 @@ public:
 	virtual void SetDone() override
 	{
 		FScreenshotRequest::OnScreenshotRequestProcessed().RemoveAll(this);
-		
+
 		UnlockViewport();
 
 		Done = true;
@@ -949,7 +949,7 @@ float HelperGetStat(FName StatName)
 		{
 			if(bCallCount)
 			{
-				return StatMessage->GetValue_CallCount(ValueType);	
+				return StatMessage->GetValue_CallCount(ValueType);
 			}
 			else
 			{
@@ -1029,7 +1029,7 @@ public:
 		, Options(InOptions)
 	{
 		UAutomationBlueprintFunctionLibrary::FinishLoadingBeforeScreenshot();
-		
+
 		WaitingFrames = 0;
 		LastLoadTime = FPlatformTime::Seconds();
 
@@ -1141,7 +1141,7 @@ public:
 		return TEXT("Waiting For Loading");
 	}
 #endif
-	
+
 private:
 	FName ExecutionFunction;
 	int32 OutputLink;
@@ -1205,7 +1205,7 @@ UAutomationEditorTask* UAutomationBlueprintFunctionLibrary::TakeHighResScreensho
 
 			Task->BindTask(MakeUnique<FScreenshotTakenState>());
 
-			// Delay taking the screenshot by a few frames			
+			// Delay taking the screenshot by a few frames
 			FTSTicker::GetCoreTicker().AddTicker(TEXT("ScreenshotDelay"), Delay, [LevelViewport, ComparisonTolerance, ComparisonNotes, Filename, ResX, ResY, bMaskEnabled, bCaptureHDR](float) {
 					FHighResScreenshotConfig& HighResScreenshotConfig = GetHighResScreenshotConfig();
 					HighResScreenshotConfig.SetResolution(ResX, ResY);
@@ -1386,4 +1386,46 @@ void UAutomationBlueprintFunctionLibrary::SetScalabilityQualityToLow(UObject* Wo
 	Scalability::SetQualityLevels(Quality, true);
 }
 
+void UAutomationBlueprintFunctionLibrary::SetEditorViewportViewMode(EViewModeIndex Index)
+{
+#if WITH_EDITOR
+	FLevelEditorModule& LevelEditorModule = FModuleManager::GetModuleChecked<FLevelEditorModule>("LevelEditor");
+
+	if (TSharedPtr<ILevelEditor> LevelEditor = LevelEditorModule.GetFirstLevelEditor())
+	{
+		for (TSharedPtr<SLevelViewport> LevelViewport : LevelEditor->GetViewports())
+		{
+			if (LevelViewport.IsValid())
+			{
+				if (TSharedPtr<FEditorViewportClient> Viewport = LevelViewport->GetViewportClient())
+				{
+					Viewport->SetViewMode(Index);
+				}
+			}
+		}
+	}
+#endif
+}
+
+void UAutomationBlueprintFunctionLibrary::SetEditorViewportVisualizeBuffer( FName BufferName )
+{
+#if WITH_EDITOR
+	FLevelEditorModule& LevelEditorModule = FModuleManager::GetModuleChecked<FLevelEditorModule>("LevelEditor");
+
+	if (TSharedPtr<ILevelEditor> LevelEditor = LevelEditorModule.GetFirstLevelEditor())
+	{
+		for (TSharedPtr<SLevelViewport> LevelViewport : LevelEditor->GetViewports())
+		{
+			if (LevelViewport.IsValid())
+			{
+				if (TSharedPtr<FEditorViewportClient> Viewport = LevelViewport->GetViewportClient())
+				{
+					Viewport->ChangeBufferVisualizationMode(BufferName);
+				}
+			}
+		}
+	}
+#endif
+}
+
 #undef LOCTEXT_NAMESPACE
diff --git a/Engine/Source/Developer/FunctionalTesting/Public/AutomationBlueprintFunctionLibrary.h b/Engine/Source/Developer/FunctionalTesting/Public/AutomationBlueprintFunctionLibrary.h
--- a/Engine/Source/Developer/FunctionalTesting/Public/AutomationBlueprintFunctionLibrary.h
+++ b/Engine/Source/Developer/FunctionalTesting/Public/AutomationBlueprintFunctionLibrary.h
@@ -189,6 +189,14 @@ public:
 
 	UFUNCTION(BlueprintCallable, Category = "Automation", meta = (HidePin = "WorldContextObject", DefaultToSelf = "WorldContextObject"))
 	static void SetScalabilityQualityToLow(UObject* WorldContextObject);
+
+	/** Sets all viewports of the first found level editor to have the given ViewMode (Lit/Unlit/etc.) **/
+	UFUNCTION(BlueprintCallable, Category = "Automation")
+	static void SetEditorViewportViewMode(EViewModeIndex Index);
+
+	/** Sets all viewports of the first found level editor to have the VisualizeBuffer ViewMode and also display a given buffer (BaseColor/Metallic/Roughness/etc.) **/
+	UFUNCTION(BlueprintCallable, Category = "Automation")
+	static void SetEditorViewportVisualizeBuffer(FName BufferName);
 };
 
 #if WITH_AUTOMATION_TESTS
