Fix build breakage due to missing include on cl 18977675.

diff --git a/Engine/Source/Developer/FunctionalTesting/Public/AutomationBlueprintFunctionLibrary.h b/Engine/Source/Developer/FunctionalTesting/Public/AutomationBlueprintFunctionLibrary.h
--- a/Engine/Source/Developer/FunctionalTesting/Public/AutomationBlueprintFunctionLibrary.h
+++ b/Engine/Source/Developer/FunctionalTesting/Public/AutomationBlueprintFunctionLibrary.h
@@ -1,14 +1,15 @@
 // Copyright Epic Games, Inc. All Rights Reserved.
 #pragma once
 
+#include "AutomationScreenshotOptions.h"
 #include "CoreMinimal.h"
-#include "UObject/ObjectMacros.h"
+#include "Engine/EngineBaseTypes.h"
 #include "Engine/LatentActionManager.h"
-#include "Kismet/BlueprintFunctionLibrary.h"
-#include "AutomationScreenshotOptions.h"
 #include "HAL/IConsoleManager.h"
-#include "Templates/UniquePtr.h"
+#include "Kismet/BlueprintFunctionLibrary.h"
 #include "Misc/AutomationTest.h"
+#include "Templates/UniquePtr.h"
+#include "UObject/ObjectMacros.h"
 #include "AutomationBlueprintFunctionLibrary.generated.h"
 
 class ACameraActor;
