Disables dynamic res on automated screenshot tests

diff --git a/Engine/Source/Developer/FunctionalTesting/Private/AutomationBlueprintFunctionLibrary.cpp b/Engine/Source/Developer/FunctionalTesting/Private/AutomationBlueprintFunctionLibrary.cpp
--- a/Engine/Source/Developer/FunctionalTesting/Private/AutomationBlueprintFunctionLibrary.cpp
+++ b/Engine/Source/Developer/FunctionalTesting/Private/AutomationBlueprintFunctionLibrary.cpp
@@ -51,6 +51,7 @@
 #include "Misc/FileHelper.h"
 #include "Materials/MaterialInterface.h"
 #include "AssetCompilingManager.h"
+#include "DynamicResolutionState.h"
 
 #if WITH_EDITOR
 #include "SLevelViewport.h"
@@ -247,6 +248,8 @@ FAutomationTestScreenshotEnvSetup::FAutomationTestScreenshotEnvSetup()
 	, TonemapperSharpen(TEXT("r.Tonemapper.Sharpen"))
 	, ScreenPercentage(TEXT("r.ScreenPercentage"))
 	, ScreenPercentageMode(TEXT("r.ScreenPercentage.Mode"))
+	, DynamicResTestScreenPercentage(TEXT("r.DynamicRes.TestScreenPercentage"))
+	, DynamicResOperationMode(TEXT("r.DynamicRes.OperationMode"))
 	, EditorViewportOverrideGameScreenPercentage(TEXT("r.Editor.Viewport.OverridePIEScreenPercentage"))
 	, SecondaryScreenPercentage(TEXT("r.SecondaryScreenPercentage.GameViewport"))
 {
@@ -287,6 +290,16 @@ void FAutomationTestScreenshotEnvSetup::Setup(UWorld* InWorld, FAutomationScreen
 		{
 			EditorViewportOverrideGameScreenPercentage.Set(0);
 		}
+
+		// Completly disable dynamic resolution
+		{
+			DynamicResTestScreenPercentage.Set(0);
+			DynamicResOperationMode.Set(0);
+
+			// Dynamic resolution status change is only taking effect at next dyn res frame.
+			GEngine->EmitDynamicResolutionEvent(EDynamicResolutionStateEvent::EndFrame);
+			GEngine->EmitDynamicResolutionEvent(EDynamicResolutionStateEvent::BeginFrame);
+		}
 		ScreenPercentageMode.Set(0);
 		ScreenPercentage.Set(100.f);
 	}
@@ -330,6 +343,8 @@ void FAutomationTestScreenshotEnvSetup::Restore()
 	//TonemapperSharpen.Restore();
 	ScreenPercentage.Restore();
 	ScreenPercentageMode.Restore();
+	DynamicResOperationMode.Restore();
+	DynamicResTestScreenPercentage.Restore();
 	EditorViewportOverrideGameScreenPercentage.Restore();
 	SecondaryScreenPercentage.Restore();
 
diff --git a/Engine/Source/Developer/FunctionalTesting/Public/AutomationBlueprintFunctionLibrary.h b/Engine/Source/Developer/FunctionalTesting/Public/AutomationBlueprintFunctionLibrary.h
--- a/Engine/Source/Developer/FunctionalTesting/Public/AutomationBlueprintFunctionLibrary.h
+++ b/Engine/Source/Developer/FunctionalTesting/Public/AutomationBlueprintFunctionLibrary.h
@@ -245,6 +245,8 @@ private:
 	FConsoleVariableSwapperTempl<float> TonemapperSharpen;
 	FConsoleVariableSwapperTempl<float> ScreenPercentage;
 	FConsoleVariableSwapperTempl<int32> ScreenPercentageMode;
+	FConsoleVariableSwapperTempl<int32> DynamicResTestScreenPercentage;
+	FConsoleVariableSwapperTempl<int32> DynamicResOperationMode;
 	FConsoleVariableSwapperTempl<int32> EditorViewportOverrideGameScreenPercentage;
 	FConsoleVariableSwapperTempl<float> SecondaryScreenPercentage;
 
