Fix Functional test not running when Navigation System is disabled

diff --git a/Engine/Source/Developer/FunctionalTesting/Classes/FunctionalTestingManager.h b/Engine/Source/Developer/FunctionalTesting/Classes/FunctionalTestingManager.h
--- a/Engine/Source/Developer/FunctionalTesting/Classes/FunctionalTestingManager.h
+++ b/Engine/Source/Developer/FunctionalTesting/Classes/FunctionalTestingManager.h
@@ -73,7 +73,6 @@ protected:
 	bool bIsRunning;
 	bool bFinished;
 	bool bLooped;
-	bool bInitialDelayApplied;
     bool bIsTearingDown;
 	uint32 CurrentIteration;
 
diff --git a/Engine/Source/Developer/FunctionalTesting/Private/FunctionalTestingManager.cpp b/Engine/Source/Developer/FunctionalTesting/Private/FunctionalTestingManager.cpp
--- a/Engine/Source/Developer/FunctionalTesting/Private/FunctionalTestingManager.cpp
+++ b/Engine/Source/Developer/FunctionalTesting/Private/FunctionalTestingManager.cpp
@@ -63,7 +63,6 @@ UFunctionalTestingManager::UFunctionalTestingManager( const FObjectInitializer&
 	, bIsRunning(false)
 	, bFinished(false)
 	, bLooped(false)
-	, bInitialDelayApplied(false)
     , bIsTearingDown(false)
 	, CurrentIteration(INDEX_NONE)
 {
@@ -160,13 +159,13 @@ void UFunctionalTestingManager::TriggerFirstValidTest()
 {
 	UWorld* World = GetWorld();
 	check(World);
-	bIsRunning = World->GetNavigationSystem() != nullptr;
+	bIsRunning = true;
 
 	const bool bIsWorldInitialized =
 		World->AreActorsInitialized() &&
-		!UNavigationSystemV1::IsNavigationBeingBuilt(World);
+			(!World->GetWorldSettings()->IsNavigationSystemEnabled() || !UNavigationSystemV1::IsNavigationBeingBuilt(World));
 
-	if (bInitialDelayApplied == true && bIsWorldInitialized)
+	if (bIsWorldInitialized)
 	{
 		bIsRunning = RunFirstValidTest();
 		if (bIsRunning == false)
@@ -176,7 +175,6 @@ void UFunctionalTestingManager::TriggerFirstValidTest()
 	}
 	else
 	{
-		bInitialDelayApplied = true;
 		static const float WaitingTime = 0.25f;
 		World->GetTimerManager().SetTimer(TriggerFirstValidTestTimerHandle, this, &UFunctionalTestingManager::TriggerFirstValidTest, WaitingTime);
 	}
