Revert changelist 19968686

#rnx

diff --git a/Engine/Source/Developer/FunctionalTesting/Classes/FunctionalTestingManager.h b/Engine/Source/Developer/FunctionalTesting/Classes/FunctionalTestingManager.h
--- a/Engine/Source/Developer/FunctionalTesting/Classes/FunctionalTestingManager.h
+++ b/Engine/Source/Developer/FunctionalTesting/Classes/FunctionalTestingManager.h
@@ -73,6 +73,7 @@ protected:
 	bool bIsRunning;
 	bool bFinished;
 	bool bLooped;
+	bool bInitialDelayApplied;
     bool bIsTearingDown;
 	uint32 CurrentIteration;
 
diff --git a/Engine/Source/Developer/FunctionalTesting/Private/FunctionalTestingManager.cpp b/Engine/Source/Developer/FunctionalTesting/Private/FunctionalTestingManager.cpp
--- a/Engine/Source/Developer/FunctionalTesting/Private/FunctionalTestingManager.cpp
+++ b/Engine/Source/Developer/FunctionalTesting/Private/FunctionalTestingManager.cpp
@@ -63,6 +63,7 @@ UFunctionalTestingManager::UFunctionalTestingManager( const FObjectInitializer&
 	, bIsRunning(false)
 	, bFinished(false)
 	, bLooped(false)
+	, bInitialDelayApplied(false)
     , bIsTearingDown(false)
 	, CurrentIteration(INDEX_NONE)
 {
@@ -159,13 +160,13 @@ void UFunctionalTestingManager::TriggerFirstValidTest()
 {
 	UWorld* World = GetWorld();
 	check(World);
-	bIsRunning = true;
+	bIsRunning = World->GetNavigationSystem() != nullptr;
 
 	const bool bIsWorldInitialized =
 		World->AreActorsInitialized() &&
-			(!World->GetWorldSettings()->IsNavigationSystemEnabled() || !UNavigationSystemV1::IsNavigationBeingBuilt(World));
+		!UNavigationSystemV1::IsNavigationBeingBuilt(World);
 
-	if (bIsWorldInitialized)
+	if (bInitialDelayApplied == true && bIsWorldInitialized)
 	{
 		bIsRunning = RunFirstValidTest();
 		if (bIsRunning == false)
@@ -175,6 +176,7 @@ void UFunctionalTestingManager::TriggerFirstValidTest()
 	}
 	else
 	{
+		bInitialDelayApplied = true;
 		static const float WaitingTime = 0.25f;
 		World->GetTimerManager().SetTimer(TriggerFirstValidTestTimerHandle, this, &UFunctionalTestingManager::TriggerFirstValidTest, WaitingTime);
 	}
