Fix Functional test not running when Navigation System is disabled (second try)

#rnx

diff --git a/Engine/Source/Developer/FunctionalTesting/Private/FunctionalTestingManager.cpp b/Engine/Source/Developer/FunctionalTesting/Private/FunctionalTestingManager.cpp
--- a/Engine/Source/Developer/FunctionalTesting/Private/FunctionalTestingManager.cpp
+++ b/Engine/Source/Developer/FunctionalTesting/Private/FunctionalTestingManager.cpp
@@ -160,11 +160,11 @@ void UFunctionalTestingManager::TriggerFirstValidTest()
 {
 	UWorld* World = GetWorld();
 	check(World);
-	bIsRunning = World->GetNavigationSystem() != nullptr;
+	bIsRunning = true;
 
 	const bool bIsWorldInitialized =
 		World->AreActorsInitialized() &&
-		!UNavigationSystemV1::IsNavigationBeingBuilt(World);
+		(!World->GetWorldSettings()->IsNavigationSystemEnabled() || !UNavigationSystemV1::IsNavigationBeingBuilt(World));
 
 	if (bInitialDelayApplied == true && bIsWorldInitialized)
 	{
