Removed useless stack walks for test info messages; added some test stats

diff --git a/Engine/Source/Developer/FunctionalTesting/Classes/FunctionalTest.h b/Engine/Source/Developer/FunctionalTesting/Classes/FunctionalTest.h
--- a/Engine/Source/Developer/FunctionalTesting/Classes/FunctionalTest.h
+++ b/Engine/Source/Developer/FunctionalTesting/Classes/FunctionalTest.h
@@ -7,6 +7,7 @@
 #include "UObject/UObjectGlobals.h"
 #include "UObject/Object.h"
 #include "Math/RandomStream.h"
+#include "Stats/Stats2.h"
 #include "GameFramework/Actor.h"
 #include "ProfilingDebugging/ExternalProfiler.h"
 #include "Math/StatisticalFloat.h"
@@ -17,6 +18,8 @@ class UBillboardComponent;
 class UTraceQueryTestResults;
 class UWorld;
 
+DECLARE_STATS_GROUP(TEXT("FunctionalTest"), STATGROUP_FunctionalTest, STATCAT_Advanced);
+
 #if UE_EXTERNAL_PROFILING_ENABLED
 //Experimental effort at automated cpu captures from the functional testing.
 class FFunctionalTestExternalProfiler : public FScopedExternalProfilerBase
diff --git a/Engine/Source/Developer/FunctionalTesting/Private/FunctionalTest.cpp b/Engine/Source/Developer/FunctionalTesting/Private/FunctionalTest.cpp
--- a/Engine/Source/Developer/FunctionalTesting/Private/FunctionalTest.cpp
+++ b/Engine/Source/Developer/FunctionalTesting/Private/FunctionalTest.cpp
@@ -26,6 +26,12 @@
 #include "Misc/RuntimeErrors.h"
 #include "FunctionalTestBase.h"
 
+DECLARE_CYCLE_STAT(TEXT("FunctionalTest - RunTest"), STAT_FunctionalTest_RunTest, STATGROUP_FunctionalTest);
+DECLARE_CYCLE_STAT(TEXT("FunctionalTest - StartTest"), STAT_FunctionalTest_StartTest, STATGROUP_FunctionalTest);
+DECLARE_CYCLE_STAT(TEXT("FunctionalTest - PrepareTest"), STAT_FunctionalTest_PrepareTest, STATGROUP_FunctionalTest);
+DECLARE_CYCLE_STAT(TEXT("FunctionalTest - Tick"), STAT_FunctionalTest_TickTest, STATGROUP_FunctionalTest);
+DECLARE_CYCLE_STAT(TEXT("FunctionalTest - FinishTest"), STAT_FunctionalTest_FinishTest, STATGROUP_FunctionalTest);
+
 namespace
 {
 	template <typename T>
@@ -217,6 +223,7 @@ void AFunctionalTest::OnConstruction(const FTransform& Transform)
 
 bool AFunctionalTest::RunTest(const TArray<FString>& Params)
 {
+	SCOPE_CYCLE_COUNTER(STAT_FunctionalTest_RunTest);
 	UWorld* World = GetWorld();
 	ensure(World->HasBegunPlay());
 
@@ -271,11 +278,13 @@ bool AFunctionalTest::RunTest(const TArray<FString>& Params)
 
 void AFunctionalTest::PrepareTest()
 {
+	SCOPE_CYCLE_COUNTER(STAT_FunctionalTest_PrepareTest);
 	ReceivePrepareTest();
 }
 
 void AFunctionalTest::StartTest()
 {
+	SCOPE_CYCLE_COUNTER(STAT_FunctionalTest_StartTest);
 	TotalTime = 0.f;
 	StartFrame = GFrameNumber;
 	StartTime = GetWorld()->GetTimeSeconds();
@@ -309,7 +318,8 @@ void AFunctionalTest::Tick(float DeltaSeconds)
 	{
 		return;
 	}
-
+	SCOPE_CYCLE_COUNTER(STAT_FunctionalTest_TickTest);
+	
 	//Do not collect garbage during the test. We force GC at the end.
 	GEngine->DelayGarbageCollection();
 
@@ -355,6 +365,7 @@ void AFunctionalTest::FinishTest(EFunctionalTestResult TestResult, const FString
 		// ignore
 		return;
 	}
+	SCOPE_CYCLE_COUNTER(STAT_FunctionalTest_FinishTest);
 	
 	// Do reporting first. When we start cleaning things up internal states that capture results
 	// are reset.
@@ -373,12 +384,14 @@ void AFunctionalTest::FinishTest(EFunctionalTestResult TestResult, const FString
 			break;
 			
 		default:
-			LogStep(ELogVerbosity::Log, *Message);
+			if (!Message.IsEmpty())
+			{
+				LogStep(ELogVerbosity::Log, *Message);
+			}
 			break;
 	}
 	
-	FFunctionalTestBase* FunctionalTest = static_cast<FFunctionalTestBase*>(FAutomationTestFramework::Get().GetCurrentTest());
-	if (FunctionalTest)
+	if (FFunctionalTestBase* FunctionalTest = static_cast<FFunctionalTestBase*>(FAutomationTestFramework::Get().GetCurrentTest()))
 	{
 		FunctionalTest->SetFunctionalTestComplete(TestLabel);
 	}
diff --git a/Engine/Source/Runtime/Core/Private/Misc/AutomationTest.cpp b/Engine/Source/Runtime/Core/Private/Misc/AutomationTest.cpp
--- a/Engine/Source/Runtime/Core/Private/Misc/AutomationTest.cpp
+++ b/Engine/Source/Runtime/Core/Private/Misc/AutomationTest.cpp
@@ -936,7 +936,7 @@ int32 FAutomationTestExecutionInfo::RemoveAllEvents(TFunctionRef<bool(FAutomatio
 	return TotalRemoved;
 }
 
-void FAutomationTestExecutionInfo::AddEvent(const FAutomationEvent& Event, int StackOffset)
+void FAutomationTestExecutionInfo::AddEvent(const FAutomationEvent& Event, int StackOffset, bool bCaptureStack)
 {
 	switch (Event.Type)
 	{
@@ -949,7 +949,7 @@ void FAutomationTestExecutionInfo::AddEvent(const FAutomationEvent& Event, int S
 	}
 
 	int32 EntryIndex = -1;
-	if (FAutomationTestFramework::Get().GetCaptureStack())
+	if (FAutomationTestFramework::Get().GetCaptureStack() && bCaptureStack)
 	{
 		SAFE_GETSTACK(Stack, StackOffset + 1, 1);
 		if (Stack.Num())
@@ -1070,10 +1070,10 @@ void FAutomationTestBase::AddWarning( const FString& InWarning, int32 StackOffse
 	}
 }
 
-void FAutomationTestBase::AddInfo( const FString& InLogItem, int32 StackOffset )
+void FAutomationTestBase::AddInfo( const FString& InLogItem, int32 StackOffset, bool bCaptureStack )
 {
 	FScopeLock Lock(&ActionCS);
-	ExecutionInfo.AddEvent(FAutomationEvent(EAutomationEventType::Info, InLogItem), StackOffset + 1);
+	ExecutionInfo.AddEvent(FAutomationEvent(EAutomationEventType::Info, InLogItem), StackOffset + 1, bCaptureStack);
 }
 
 void FAutomationTestBase::AddAnalyticsItem(const FString& InAnalyticsItem)
@@ -1102,10 +1102,10 @@ void FAutomationTestBase::SetTelemetryStorage(const FString& StorageName)
 	ExecutionInfo.TelemetryStorage = StorageName;
 }
 
-void FAutomationTestBase::AddEvent(const FAutomationEvent& InEvent, int32 StackOffset)
+void FAutomationTestBase::AddEvent(const FAutomationEvent& InEvent, int32 StackOffset, bool bCaptureStack)
 {
 	FScopeLock Lock(&ActionCS);
-	ExecutionInfo.AddEvent(InEvent, StackOffset + 1);
+	ExecutionInfo.AddEvent(InEvent, StackOffset + 1, bCaptureStack);
 }
 
 bool FAutomationTestBase::HasAnyErrors() const
diff --git a/Engine/Source/Runtime/Core/Public/Misc/AutomationTest.h b/Engine/Source/Runtime/Core/Public/Misc/AutomationTest.h
--- a/Engine/Source/Runtime/Core/Public/Misc/AutomationTest.h
+++ b/Engine/Source/Runtime/Core/Public/Misc/AutomationTest.h
@@ -175,7 +175,7 @@ public:
 	/** Any errors that occurred during execution */
 	const TArray<FAutomationExecutionEntry>& GetEntries() const { return Entries; }
 
-	void AddEvent(const FAutomationEvent& Event, int StackOffset = 0);
+	void AddEvent(const FAutomationEvent& Event, int StackOffset = 0, bool bCaptureStack = true);
 
 	void AddWarning(const FString& WarningMessage);
 	void AddError(const FString& ErrorMessage);
@@ -1278,14 +1278,14 @@ public:
 	 *
 	 * @param	InLogItem	Log item to add to this test
 	 */
-	virtual void AddInfo( const FString& InLogItem, int32 StackOffset = 0);
+	virtual void AddInfo( const FString& InLogItem, int32 StackOffset = 0, bool bCaptureStack = false);
 
 	/**
 	 * Adds an automation event directly into the execution log.
 	 *
 	 * @param	InLogItem	Log item to add to this test
 	 */
-	virtual void AddEvent(const FAutomationEvent& InEvent, int32 StackOffset = 0);
+	virtual void AddEvent(const FAutomationEvent& InEvent, int32 StackOffset = 0, bool bCaptureStack = false);
 
 	/**
 	 * Adds a analytics string to parse later
