OFPA support for functional tests.
Maps and functional test actors needs to be resaved.

diff --git a/Engine/Source/Developer/FunctionalTesting/Classes/FunctionalTest.h b/Engine/Source/Developer/FunctionalTesting/Classes/FunctionalTest.h
--- a/Engine/Source/Developer/FunctionalTesting/Classes/FunctionalTest.h
+++ b/Engine/Source/Developer/FunctionalTesting/Classes/FunctionalTest.h
@@ -557,7 +557,7 @@ public:
 
 public:
 	
-	/** ACtors registered this way will be automatically destroyed (by limiting their lifespan)
+	/** Actors registered this way will be automatically destroyed (by limiting their lifespan)
 	 *	on test finish */
 	UFUNCTION(BlueprintCallable, Category="Development", meta=(Keywords = "Delete"))
 	virtual void RegisterAutoDestroyActor(AActor* ActorToAutoDestroy);
@@ -571,6 +571,7 @@ public:
 
 #if WITH_EDITOR
 	void PostEditChangeProperty(struct FPropertyChangedEvent& PropertyChangedEvent) override;
+	virtual void GetAssetRegistryTags(TArray<FAssetRegistryTag>& OutTags) const override;
 
 	static void OnSelectObject(UObject* NewSelection);
 #endif // WITH_EDITOR
@@ -579,6 +580,9 @@ public:
 	virtual void OnConstruction(const FTransform& Transform) override;
 	virtual void Tick(float DeltaSeconds) override;
 	virtual void EndPlay(const EEndPlayReason::Type EndPlayReason) override;
+#if WITH_EDITOR
+	virtual bool CanChangeIsSpatiallyLoadedFlag() const { return false; }
+#endif
 	// AActor interface end
 
 	UFUNCTION(BlueprintCallable, Category = "Functional Testing")
diff --git a/Engine/Source/Developer/FunctionalTesting/Private/FunctionalTest.cpp b/Engine/Source/Developer/FunctionalTesting/Private/FunctionalTest.cpp
--- a/Engine/Source/Developer/FunctionalTesting/Private/FunctionalTest.cpp
+++ b/Engine/Source/Developer/FunctionalTesting/Private/FunctionalTest.cpp
@@ -193,6 +193,8 @@ AFunctionalTest::AFunctionalTest( const FObjectInitializer& ObjectInitializer )
 		TestName->SetRelativeRotation(FRotator(0, 0, 0));
 		TestName->SetupAttachment(RootComponent);
 	}
+
+	bIsSpatiallyLoaded = false;
 #endif
 }
 
@@ -552,6 +554,18 @@ void AFunctionalTest::PostEditChangeProperty( struct FPropertyChangedEvent& Prop
 	}
 }
 
+void AFunctionalTest::GetAssetRegistryTags(TArray<FAssetRegistryTag>& OutTags) const
+{
+	Super::GetAssetRegistryTags(OutTags);
+
+	if (IsPackageExternal() && IsEnabled())
+	{
+		const FString TestActor = GetActorLabel() + TEXT("|") + GetName();
+		const TCHAR* TestCategory = IsEditorOnlyObject(this) ? TEXT("TestNameEditor") : TEXT("TestName");
+		OutTags.Add(UObject::FAssetRegistryTag(TestCategory, TestActor, UObject::FAssetRegistryTag::TT_Hidden));
+	}
+}
+
 void AFunctionalTest::OnSelectObject(UObject* NewSelection)
 {
 	AFunctionalTest* AsFTest = Cast<AFunctionalTest>(NewSelection);
diff --git a/Engine/Source/Developer/FunctionalTesting/Private/FunctionalTestingModule.cpp b/Engine/Source/Developer/FunctionalTesting/Private/FunctionalTestingModule.cpp
--- a/Engine/Source/Developer/FunctionalTesting/Private/FunctionalTestingModule.cpp
+++ b/Engine/Source/Developer/FunctionalTesting/Private/FunctionalTestingModule.cpp
@@ -62,31 +62,27 @@ void FFunctionalTestingModule::ShutdownModule()
 void FFunctionalTestingModule::OnGetAssetTagsForWorld(const UWorld* World, TArray<UObject::FAssetRegistryTag>& OutTags)
 {
 #if WITH_EDITOR
-	int32 Tests = 0;
 	FString TestNames, TestNamesEditor;
 	for (TActorIterator<AFunctionalTest> ActorItr(const_cast<UWorld*>(World), AFunctionalTest::StaticClass(), EActorIteratorFlags::AllActors); ActorItr; ++ActorItr)
 	{
 		AFunctionalTest* FunctionalTest = *ActorItr;
 
-		// Only include enabled tests in the list of functional tests to run.
-		if (FunctionalTest->IsEnabled())
+		if (!FunctionalTest->IsPackageExternal())
 		{
-			bool bIsEditorOnly = IsEditorOnlyObject(FunctionalTest);
+			// Only include enabled tests in the list of functional tests to run.
+			if (FunctionalTest->IsEnabled())
+			{
+				bool bIsEditorOnly = IsEditorOnlyObject(FunctionalTest);
 
-			// Check if this class is editor only
-			FString& NamesAppend = bIsEditorOnly ? TestNamesEditor : TestNames;
+				// Check if this class is editor only
+				FString& NamesAppend = bIsEditorOnly ? TestNamesEditor : TestNames;
 
-			Tests++;
-			NamesAppend.Append(FunctionalTest->GetActorLabel() + TEXT("|") + FunctionalTest->GetName());
-			NamesAppend.Append(TEXT(";"));
+				NamesAppend.Append(FunctionalTest->GetActorLabel() + TEXT("|") + FunctionalTest->GetName());
+				NamesAppend.Append(TEXT(";"));
+			}
 		}
 	}
 
-	if (Tests > 0)
-	{
-		OutTags.Add(UObject::FAssetRegistryTag("Tests", FString::FromInt(Tests), UObject::FAssetRegistryTag::TT_Numerical));
-	}
-
 	if (!TestNames.IsEmpty())
 	{
 		OutTags.Add(UObject::FAssetRegistryTag("TestNames", TestNames, UObject::FAssetRegistryTag::TT_Hidden));
@@ -133,28 +129,62 @@ void FFunctionalTestingModule::GetMapTests(bool bEditorOnlyTests, TArray<FString
 				FString MapPackageName = MapAsset.PackageName.ToString();
 				if (!IsDeveloperDirectoryIncluded && MapPackageName.Find(TEXT("/Game/Developers")) == 0) continue;
 
-				FAssetDataTagMapSharedView::FFindTagResult Tests = MapAsset.TagsAndValues.FindTag(TEXT("Tests"));
-				FAssetDataTagMapSharedView::FFindTagResult TestNames = MapAsset.TagsAndValues.FindTag(bEditorOnlyTests ? TEXT("TestNamesEditor") : TEXT("TestNames"));
+				FString AllTestNames;
+				FAssetDataTagMapSharedView::FFindTagResult MapTestNames = MapAsset.TagsAndValues.FindTag(bEditorOnlyTests ? TEXT("TestNamesEditor") : TEXT("TestNames"));
+
+				if (MapTestNames.IsSet())
+				{
+					AllTestNames = MapTestNames.GetValue();
+				}
 
-				if (Tests.IsSet() && TestNames.IsSet())
+#if WITH_EDITOR
+				// Also append external functional test actors
+				if (ULevel::GetIsLevelUsingExternalActorsFromAsset(MapAsset))
 				{
-					int32 TestCount = FCString::Atoi(*Tests.GetValue());
-					if (TestCount > 0)
+					const FString LevelExternalActorsPath = ULevel::GetExternalActorsPath(MapPackageName);
+
+					// Do a synchronous scan of the level external actors path.			
+					AssetRegistry.ScanPathsSynchronous({ LevelExternalActorsPath }, /*bForceRescan*/false, /*bIgnoreDenyListScanFilters*/false);
+
+					FARFilter ActorsFilter;
+					ActorsFilter.bRecursivePaths = true;
+					ActorsFilter.bIncludeOnlyOnDiskAssets = true;
+					ActorsFilter.PackagePaths.Add(*LevelExternalActorsPath);
+
+					TArray<FAssetData> ActorList;
+					AssetRegistry.GetAssets(ActorsFilter, ActorList);
+
+					for (const FAssetData& ActorAsset : ActorList)
 					{
-						TArray<FString> MapTests;
-						TestNames.GetValue().ParseIntoArray(MapTests, TEXT(";"), true);
+						FAssetDataTagMapSharedView::FFindTagResult ActorTestName = ActorAsset.TagsAndValues.FindTag(bEditorOnlyTests ? TEXT("TestNameEditor") : TEXT("TestName"));
 
-						for (const FString& MapTest : MapTests)
+						if (ActorTestName.IsSet())
 						{
-							FString BeautifulTestName;
-							FString RealTestName;
-
-							if (MapTest.Split(TEXT("|"), &BeautifulTestName, &RealTestName))
+							if (!AllTestNames.IsEmpty())
 							{
-								OutBeautifiedNames.Add(MapPackageName + TEXT(".") + *BeautifulTestName);
-								OutTestCommands.Add(MapAssetPath + TEXT(";") + MapAsset.PackageName.ToString() + TEXT(";") + *RealTestName);
-								OutTestMapAssets.AddUnique(MapAssetPath);
+								AllTestNames += TEXT(";");
 							}
+
+							AllTestNames += ActorTestName.GetValue();
+						}
+					}
+				}
+#endif
+				if (!AllTestNames.IsEmpty())
+				{
+					TArray<FString> MapTests;
+					AllTestNames.ParseIntoArray(MapTests, TEXT(";"), true);
+
+					for (const FString& MapTest : MapTests)
+					{
+						FString BeautifulTestName;
+						FString RealTestName;
+
+						if (MapTest.Split(TEXT("|"), &BeautifulTestName, &RealTestName))
+						{
+							OutBeautifiedNames.Add(MapPackageName + TEXT(".") + *BeautifulTestName);
+							OutTestCommands.Add(MapAssetPath + TEXT(";") + MapAsset.PackageName.ToString() + TEXT(";") + *RealTestName);
+							OutTestMapAssets.AddUnique(MapAssetPath);
 						}
 					}
 				}
