Fixed 2 problems when hitting stop test button,
* the test aren't stop correctly. So the second time you run it, it give you the error that the test is already running.
*  it does not stop PIE even if it started it. But the automation controller manager finishes all the test, then it stops correctly PIE.

diff --git a/Engine/Source/Developer/AutomationController/Private/AutomationControllerManger.cpp b/Engine/Source/Developer/AutomationController/Private/AutomationControllerManger.cpp
--- a/Engine/Source/Developer/AutomationController/Private/AutomationControllerManger.cpp
+++ b/Engine/Source/Developer/AutomationController/Private/AutomationControllerManger.cpp
@@ -375,6 +375,14 @@ void FAutomationControllerManager::StopTests()
 	}
 
 	TestRunningArray.Empty();
+
+	// Close play window
+#if WITH_EDITOR
+	if (GUnrealEd)
+	{
+		GUnrealEd->RequestEndPlayMap();
+	}
+#endif
 }
 
 void FAutomationControllerManager::Init()
diff --git a/Engine/Source/Developer/FunctionalTesting/Private/FunctionalTest.cpp b/Engine/Source/Developer/FunctionalTesting/Private/FunctionalTest.cpp
--- a/Engine/Source/Developer/FunctionalTesting/Private/FunctionalTest.cpp
+++ b/Engine/Source/Developer/FunctionalTesting/Private/FunctionalTest.cpp
@@ -433,7 +433,8 @@ void AFunctionalTest::EndPlay(const EEndPlayReason::Type EndPlayReason)
 	// If end play occurs and we're still running, notify that the testing has stopped.
 	if (bIsRunning)
 	{
-		TestFinishedObserver.ExecuteIfBound(this);
+		// Tell the test he is being aborted
+		FinishTest(EFunctionalTestResult::Invalid, TEXT("Test was aborted"));
 	}
 
 	TestFinishedObserver.Unbind();
