Added missing tests for various core types. Also cleaned up const and reference usage in the function declarations.

diff --git a/Engine/Source/Developer/FunctionalTesting/Classes/FunctionalTest.h b/Engine/Source/Developer/FunctionalTesting/Classes/FunctionalTest.h
--- a/Engine/Source/Developer/FunctionalTesting/Classes/FunctionalTest.h
+++ b/Engine/Source/Developer/FunctionalTesting/Classes/FunctionalTest.h
@@ -400,42 +400,42 @@ public:
 	 * @param What	A name to use in the message if the assert fails ("Expected 'What' to be {Expected} but it was {Actual} for context ''")
 	 */
 	UFUNCTION(BlueprintCallable, Category = "Asserts", DisplayName = "Assert Equal (Transform)", meta = ( HidePin = "ContextObject", DefaultToSelf = "ContextObject"))
-	virtual bool AssertEqual_Transform(const FTransform& Actual, const FTransform& Expected, const FString& What, float Tolerance = 1.e-4, const UObject* ContextObject = nullptr);
+	virtual bool AssertEqual_Transform(const FTransform& Actual, const FTransform& Expected, const FString& What, float Tolerance = 1.e-4f, const UObject* ContextObject = nullptr);
 
 	/**
 	 * Assert that two floats are equal within tolerance between two floats.
 	 * @param What	A name to use in the message if the assert fails (What: expected {Actual} to be Equal To {Expected} within Tolerance for context '')
 	 */
 	UFUNCTION(BlueprintCallable, Category = "Asserts", DisplayName = "Assert Equal (Float)", meta = ( HidePin = "ContextObject", DefaultToSelf = "ContextObject"))
-	virtual bool AssertEqual_Float(const float Actual, const float Expected, const FString& What, const float Tolerance = 1.e-4, const UObject* ContextObject = nullptr);
+	virtual bool AssertEqual_Float(float Actual, float Expected, const FString& What, float Tolerance = 1.e-4f, const UObject* ContextObject = nullptr);
 
 	/**
 	 * Assert that two double are equal within tolerance between two doubles.
 	 * @param What	A name to use in the message if the assert fails (What: expected {Actual} to be Equal To {Expected} within Tolerance for context '')
 	 */
 	UFUNCTION(BlueprintCallable, Category = "Asserts", DisplayName = "Assert Equal (Double)", meta = (HidePin = "ContextObject", DefaultToSelf = "ContextObject"))
-	bool AssertEqual_Double(const double Actual, const double Expected, const FString& What, const double Tolerance = 1.e-4, const UObject* ContextObject = nullptr);
+	bool AssertEqual_Double(double Actual, double Expected, const FString& What, double Tolerance = 1.e-4, const UObject* ContextObject = nullptr);
 
 	/**
 	* Assert that two bools are equal
 	* @param What	A name to use in the message if the assert fails (What: expected {Actual} to be Equal To {Expected} for context '')
 	*/
 	UFUNCTION(BlueprintCallable, Category = "Asserts", DisplayName = "Assert Equal (Bool)", meta = (HidePin = "ContextObject", DefaultToSelf = "ContextObject"))
-	virtual bool AssertEqual_Bool(const bool Actual, const bool Expected, const FString& What, const UObject* ContextObject = nullptr);
+	virtual bool AssertEqual_Bool(bool Actual, bool Expected, const FString& What, const UObject* ContextObject = nullptr);
 
 	/**
 	* Assert that two ints are equal
 	* @param What	A name to use in the message if the assert fails (What: expected {Actual} to be Equal To {Expected} for context '')
 	*/
 	UFUNCTION(BlueprintCallable, Category = "Asserts", DisplayName = "Assert Equal (Integer)", meta = (HidePin = "ContextObject", DefaultToSelf = "ContextObject"))
-	virtual bool AssertEqual_Int(const int Actual, const int Expected, const FString& What, const UObject* ContextObject = nullptr);
+	virtual bool AssertEqual_Int(int Actual, int Expected, const FString& What, const UObject* ContextObject = nullptr);
 
 	/**
 	* Assert that two FNames are equal
 	* @param What	A name to use in the message if the assert fails (What: expected {Actual} to be Equal To {Expected} for context '')
 	*/
 	UFUNCTION(BlueprintCallable, Category = "Asserts", DisplayName = "Assert Equal (FName)", meta = (HidePin = "ContextObject", DefaultToSelf = "ContextObject"))
-	virtual bool AssertEqual_Name(const FName Actual, const FName Expected, const FString& What, const UObject* ContextObject = nullptr);
+	virtual bool AssertEqual_Name(FName Actual, FName Expected, const FString& What, const UObject* ContextObject = nullptr);
 
 	/**
 	* Assert that two Objects are equal
@@ -456,42 +456,112 @@ public:
 	 * @param What	A name to use in the message if the assert fails ("Expected 'What' to be {Expected} but it was {Actual} for context ''")
 	 */
 	UFUNCTION(BlueprintCallable, Category = "Asserts", DisplayName = "Assert Equal (Rotator)", meta = ( HidePin = "ContextObject", DefaultToSelf = "ContextObject"))
-	virtual bool AssertEqual_Rotator(const FRotator Actual, const FRotator Expected, const FString& What, const float Tolerance = 1.e-4, const UObject* ContextObject = nullptr);
+	virtual bool AssertEqual_Rotator(const FRotator& Actual, const FRotator& Expected, const FString& What, float Tolerance = 1.e-4f, const UObject* ContextObject = nullptr);
 
 	/**
 	 * Assert that the component angles of two rotators are all not equal within a small tolerance.
 	 * @param What	A name to use in the message if the assert fails ("Expected 'What' not to be {Expected} but it was {Actual} for context ''")
 	 */
 	UFUNCTION(BlueprintCallable, Category = "Asserts", DisplayName = "Assert Not Equal (Rotator)", meta = ( HidePin = "ContextObject", DefaultToSelf = "ContextObject"))
-	virtual bool AssertNotEqual_Rotator(const FRotator Actual, const FRotator NotExpected, const FString& What, const UObject* ContextObject = nullptr);
+	virtual bool AssertNotEqual_Rotator(const FRotator& Actual, const FRotator& NotExpected, const FString& What, const UObject* ContextObject = nullptr);
 
 	/**
 	 * Assert that two vectors are (memberwise) equal within a small tolerance.
 	 * @param What	A name to use in the message if the assert fails ("Expected 'What' to be {Expected} but it was {Actual} for context ''")
 	 */
 	UFUNCTION(BlueprintCallable, Category = "Asserts", DisplayName = "Assert Equal (Vector)", meta = ( HidePin = "ContextObject", DefaultToSelf = "ContextObject"))
-	virtual bool AssertEqual_Vector(const FVector Actual, const FVector Expected, const FString& What, const float Tolerance = 1.e-4f, const UObject* ContextObject = nullptr);
+	virtual bool AssertEqual_Vector(const FVector& Actual, const FVector& Expected, const FString& What, float Tolerance = 1.e-4f, const UObject* ContextObject = nullptr);
 
 	/**
 	 * Assert that two vectors are (memberwise) not equal within a small tolerance.
 	 * @param What	A name to use in the message if the assert fails ("Expected 'What' not to be {Expected} but it was {Actual} for context ''")
 	 */
 	UFUNCTION(BlueprintCallable, Category = "Asserts", DisplayName = "Assert Not Equal (Vector)", meta = ( HidePin = "ContextObject", DefaultToSelf = "ContextObject"))
-	virtual bool AssertNotEqual_Vector(const FVector Actual, const FVector NotExpected, const FString& What, const UObject* ContextObject = nullptr);
+	virtual bool AssertNotEqual_Vector(const FVector& Actual, const FVector& NotExpected, const FString& What, const UObject* ContextObject = nullptr);
+
+	/**
+	 * Assert that two two-component vectors are (memberwise) equal within a small tolerance.
+	 * @param What	A name to use in the message if the assert fails ("Expected 'What' to be {Expected} but it was {Actual} for context ''")
+	 */
+	UFUNCTION(BlueprintCallable, Category = "Asserts", DisplayName = "Assert Equal (Vector2D)", meta = ( HidePin = "ContextObject", DefaultToSelf = "ContextObject"))
+	virtual bool AssertEqual_Vector2D(const FVector2D& Actual, const FVector2D& Expected, const FString& What, float Tolerance = 1.e-4f, const UObject* ContextObject = nullptr);
+
+	/**
+	 * Assert that two two-component vectors are (memberwise) not equal within a small tolerance.
+	 * @param What	A name to use in the message if the assert fails ("Expected 'What' not to be {Expected} but it was {Actual} for context ''")
+	 */
+	UFUNCTION(BlueprintCallable, Category = "Asserts", DisplayName = "Assert Not Equal (Vector2D)", meta = ( HidePin = "ContextObject", DefaultToSelf = "ContextObject"))
+	virtual bool AssertNotEqual_Vector2D(const FVector2D& Actual, const FVector2D& NotExpected, const FString& What, const UObject* ContextObject = nullptr);
+
+	/**
+	 * Assert that two four-component vectors are (memberwise) equal within a small tolerance.
+	 * @param What	A name to use in the message if the assert fails ("Expected 'What' to be {Expected} but it was {Actual} for context ''")
+	 */
+	UFUNCTION(BlueprintCallable, Category = "Asserts", DisplayName = "Assert Equal (Vector4)", meta = ( HidePin = "ContextObject", DefaultToSelf = "ContextObject"))
+	virtual bool AssertEqual_Vector4(const FVector4& Actual, const FVector4& Expected, const FString& What, float Tolerance = 1.e-4f, const UObject* ContextObject = nullptr);
+
+	/**
+	 * Assert that two four-component vectors are (memberwise) not equal within a small tolerance.
+	 * @param What	A name to use in the message if the assert fails ("Expected 'What' not to be {Expected} but it was {Actual} for context ''")
+	 */
+	UFUNCTION(BlueprintCallable, Category = "Asserts", DisplayName = "Assert Not Equal (Vector4)", meta = ( HidePin = "ContextObject", DefaultToSelf = "ContextObject"))
+	virtual bool AssertNotEqual_Vector4(const FVector4& Actual, const FVector4& NotExpected, const FString& What, const UObject* ContextObject = nullptr);
+
+	/**
+	 * Assert that two planes are (memberwise) equal within a small tolerance.
+	 * @param What	A name to use in the message if the assert fails ("Expected 'What' to be {Expected} but it was {Actual} for context ''")
+	 */
+	UFUNCTION(BlueprintCallable, Category = "Asserts", DisplayName = "Assert Equal (Plane)", meta = ( HidePin = "ContextObject", DefaultToSelf = "ContextObject"))
+	virtual bool AssertEqual_Plane(const FPlane& Actual, const FPlane& Expected, const FString& What, float Tolerance = 1.e-4f, const UObject* ContextObject = nullptr);
+
+	/**
+	 * Assert that two planes are (memberwise) not equal within a small tolerance.
+	 * @param What	A name to use in the message if the assert fails ("Expected 'What' not to be {Expected} but it was {Actual} for context ''")
+	 */
+	UFUNCTION(BlueprintCallable, Category = "Asserts", DisplayName = "Assert Not Equal (Plane)", meta = ( HidePin = "ContextObject", DefaultToSelf = "ContextObject"))
+	virtual bool AssertNotEqual_Plane(const FPlane& Actual, const FPlane& NotExpected, const FString& What, const UObject* ContextObject = nullptr);
+
+	/**
+	 * Assert that two quats are (memberwise) equal within a small tolerance.
+	 * @param What	A name to use in the message if the assert fails ("Expected 'What' to be {Expected} but it was {Actual} for context ''")
+	 */
+	UFUNCTION(BlueprintCallable, Category = "Asserts", DisplayName = "Assert Equal (Quat)", meta = ( HidePin = "ContextObject", DefaultToSelf = "ContextObject"))
+	virtual bool AssertEqual_Quat(const FQuat& Actual, const FQuat& Expected, const FString& What, float Tolerance = 1.e-4f, const UObject* ContextObject = nullptr);
+
+	/**
+	 * Assert that two quats are (memberwise) not equal within a small tolerance.
+	 * @param What	A name to use in the message if the assert fails ("Expected 'What' not to be {Expected} but it was {Actual} for context ''")
+	 */
+	UFUNCTION(BlueprintCallable, Category = "Asserts", DisplayName = "Assert Not Equal (Quat)", meta = ( HidePin = "ContextObject", DefaultToSelf = "ContextObject"))
+	virtual bool AssertNotEqual_Quat(const FQuat& Actual, const FQuat& NotExpected, const FString& What, const UObject* ContextObject = nullptr);
+
+	/**
+	 * Assert that two 4x4 matrices are (memberwise) equal within a small tolerance.
+	 * @param What	A name to use in the message if the assert fails ("Expected 'What' to be {Expected} but it was {Actual} for context ''")
+	 */
+	UFUNCTION(BlueprintCallable, Category = "Asserts", DisplayName = "Assert Equal (Matrix)", meta = ( HidePin = "ContextObject", DefaultToSelf = "ContextObject"))
+	virtual bool AssertEqual_Matrix(const FMatrix& Actual, const FMatrix& Expected, const FString& What, float Tolerance = 1.e-4f, const UObject* ContextObject = nullptr);
+
+	/**
+	 * Assert that two 4x4 matrices are (memberwise) not equal within a small tolerance.
+	 * @param What	A name to use in the message if the assert fails ("Expected 'What' not to be {Expected} but it was {Actual} for context ''")
+	 */
+	UFUNCTION(BlueprintCallable, Category = "Asserts", DisplayName = "Assert Not Equal (Matrix)", meta = ( HidePin = "ContextObject", DefaultToSelf = "ContextObject"))
+	virtual bool AssertNotEqual_Matrix(const FMatrix& Actual, const FMatrix& NotExpected, const FString& What, const UObject* ContextObject = nullptr);
 
 	/**
 	 * Assert that two Strings are equal.
 	 * @param What	A name to use in the message if the assert fails ("Expected 'What' to be {Expected} but it was {Actual} for context ''")
 	 */
 	UFUNCTION(BlueprintCallable, Category = "Asserts", DisplayName = "Assert Equal (String)", meta = ( HidePin = "ContextObject", DefaultToSelf = "ContextObject"))
-	virtual bool AssertEqual_String(const FString Actual, const FString Expected, const FString& What, const UObject* ContextObject = nullptr);
+	virtual bool AssertEqual_String(const FString& Actual, const FString& Expected, const FString& What, const UObject* ContextObject = nullptr);
 
 	/**
 	 * Assert that two Strings are not equal.
 	 * @param What	A name to use in the message if the assert fails ("Expected 'What' not to be {Expected} but it was {Actual} for context ''")
 	 */
 	UFUNCTION(BlueprintCallable, Category = "Asserts", DisplayName = "Assert Not Equal (String)", meta = ( HidePin = "ContextObject", DefaultToSelf = "ContextObject"))
-	virtual bool AssertNotEqual_String(const FString Actual, const FString NotExpected, const FString& What, const UObject* ContextObject = nullptr);
+	virtual bool AssertNotEqual_String(const FString& Actual, const FString& NotExpected, const FString& What, const UObject* ContextObject = nullptr);
 
 	/**
 	* Assert that two TraceQueryResults are equal.
diff --git a/Engine/Source/Developer/FunctionalTesting/Private/FunctionalTest.cpp b/Engine/Source/Developer/FunctionalTesting/Private/FunctionalTest.cpp
--- a/Engine/Source/Developer/FunctionalTesting/Private/FunctionalTest.cpp
+++ b/Engine/Source/Developer/FunctionalTesting/Private/FunctionalTest.cpp
@@ -826,7 +826,7 @@ bool AFunctionalTest::AssertNotEqual_Transform(const FTransform& Actual, const F
 	}
 }
 
-bool AFunctionalTest::AssertEqual_Rotator(const FRotator Actual, const FRotator Expected, const FString& What, const float Tolerance, const UObject* ContextObject)
+bool AFunctionalTest::AssertEqual_Rotator(const FRotator& Actual, const FRotator& Expected, const FString& What, const float Tolerance, const UObject* ContextObject)
 {
 	if ( !Expected.Equals(Actual, Tolerance) )
 	{
@@ -840,7 +840,7 @@ bool AFunctionalTest::AssertEqual_Rotator(const FRotator Actual, const FRotator
 	}
 }
 
-bool AFunctionalTest::AssertNotEqual_Rotator(const FRotator Actual, const FRotator NotExpected, const FString& What, const UObject* ContextObject)
+bool AFunctionalTest::AssertNotEqual_Rotator(const FRotator& Actual, const FRotator& NotExpected, const FString& What, const UObject* ContextObject)
 {
 	if ( NotExpected.Equals(Actual) )
 	{
@@ -854,7 +854,7 @@ bool AFunctionalTest::AssertNotEqual_Rotator(const FRotator Actual, const FRotat
 	}
 }
 
-bool AFunctionalTest::AssertEqual_Vector(const FVector Actual, const FVector Expected, const FString& What, const float Tolerance, const UObject* ContextObject)
+bool AFunctionalTest::AssertEqual_Vector(const FVector& Actual, const FVector& Expected, const FString& What, const float Tolerance, const UObject* ContextObject)
 {
 	if ( !Expected.Equals(Actual, Tolerance) )
 	{
@@ -868,7 +868,7 @@ bool AFunctionalTest::AssertEqual_Vector(const FVector Actual, const FVector Exp
 	}
 }
 
-bool AFunctionalTest::AssertNotEqual_Vector(const FVector Actual, const FVector NotExpected, const FString& What, const UObject* ContextObject)
+bool AFunctionalTest::AssertNotEqual_Vector(const FVector& Actual, const FVector& NotExpected, const FString& What, const UObject* ContextObject)
 {
 	if ( NotExpected.Equals(Actual) )
 	{
@@ -882,7 +882,148 @@ bool AFunctionalTest::AssertNotEqual_Vector(const FVector Actual, const FVector
 	}
 }
 
-bool AFunctionalTest::AssertEqual_String(const FString Actual, const FString Expected, const FString& What, const UObject* ContextObject)
+bool AFunctionalTest::AssertEqual_Vector2D(const FVector2D& Actual, const FVector2D& Expected, const FString& What, const float Tolerance, const UObject* ContextObject)
+{
+	if (!Expected.Equals(Actual, Tolerance))
+	{
+		LogStep(ELogVerbosity::Error, FString::Printf(TEXT("Expected '%s' to be {%s} but it was {%s} within tolerance {%f} for context '%s'"), *What, *Expected.ToString(), *Actual.ToString(), Tolerance, ContextObject ? *ContextObject->GetName() : TEXT("")));
+		return false;
+	}
+	else
+	{
+		LogStep(ELogVerbosity::Log, FString::Printf(TEXT("Vector2D assertion passed (%s)"), *What));
+		return true;
+	}
+}
+
+bool AFunctionalTest::AssertNotEqual_Vector2D(const FVector2D& Actual, const FVector2D& NotExpected, const FString& What, const UObject* ContextObject)
+{
+	if (NotExpected.Equals(Actual))
+	{
+		LogStep(ELogVerbosity::Error, FString::Printf(TEXT("Expected '%s' not to be {%s} for context '%s'"), *What, *NotExpected.ToString(), ContextObject ? *ContextObject->GetName() : TEXT("")));
+		return false;
+	}
+	else
+	{
+		LogStep(ELogVerbosity::Log, FString::Printf(TEXT("Vector2D assertion passed (%s)"), *What));
+		return true;
+	}
+}
+
+bool AFunctionalTest::AssertEqual_Vector4(const FVector4& Actual, const FVector4& Expected, const FString& What, const float Tolerance, const UObject* ContextObject)
+{
+	if (!Expected.Equals(Actual, Tolerance))
+	{
+		LogStep(ELogVerbosity::Error, FString::Printf(TEXT("Expected '%s' to be {%s} but it was {%s} within tolerance {%f} for context '%s'"), *What, *Expected.ToString(), *Actual.ToString(), Tolerance, ContextObject ? *ContextObject->GetName() : TEXT("")));
+		return false;
+	}
+	else
+	{
+		LogStep(ELogVerbosity::Log, FString::Printf(TEXT("Vector4 assertion passed (%s)"), *What));
+		return true;
+	}
+}
+
+bool AFunctionalTest::AssertNotEqual_Vector4(const FVector4& Actual, const FVector4& NotExpected, const FString& What, const UObject* ContextObject)
+{
+	if (NotExpected.Equals(Actual))
+	{
+		LogStep(ELogVerbosity::Error, FString::Printf(TEXT("Expected '%s' not to be {%s} for context '%s'"), *What, *NotExpected.ToString(), ContextObject ? *ContextObject->GetName() : TEXT("")));
+		return false;
+	}
+	else
+	{
+		LogStep(ELogVerbosity::Log, FString::Printf(TEXT("Vector4 assertion passed (%s)"), *What));
+		return true;
+	}
+}
+
+
+bool AFunctionalTest::AssertEqual_Plane(const FPlane& Actual, const FPlane& Expected, const FString& What, const float Tolerance, const UObject* ContextObject)
+{
+	if (!Expected.Equals(Actual, Tolerance))
+	{
+		LogStep(ELogVerbosity::Error, FString::Printf(TEXT("Expected '%s' to be {%s} but it was {%s} within tolerance {%f} for context '%s'"), *What, *Expected.ToString(), *Actual.ToString(), Tolerance, ContextObject ? *ContextObject->GetName() : TEXT("")));
+		return false;
+	}
+	else
+	{
+		LogStep(ELogVerbosity::Log, FString::Printf(TEXT("Plane assertion passed (%s)"), *What));
+		return true;
+	}
+}
+
+bool AFunctionalTest::AssertNotEqual_Plane(const FPlane& Actual, const FPlane& NotExpected, const FString& What, const UObject* ContextObject)
+{
+	if (NotExpected.Equals(Actual))
+	{
+		LogStep(ELogVerbosity::Error, FString::Printf(TEXT("Expected '%s' not to be {%s} for context '%s'"), *What, *NotExpected.ToString(), ContextObject ? *ContextObject->GetName() : TEXT("")));
+		return false;
+	}
+	else
+	{
+		LogStep(ELogVerbosity::Log, FString::Printf(TEXT("Plane assertion passed (%s)"), *What));
+		return true;
+	}
+}
+
+bool AFunctionalTest::AssertEqual_Quat(const FQuat& Actual, const FQuat& Expected, const FString& What, const float Tolerance, const UObject* ContextObject)
+{
+	if (!Expected.Equals(Actual, Tolerance))
+	{
+		LogStep(ELogVerbosity::Error, FString::Printf(TEXT("Expected '%s' to be {%s} but it was {%s} within tolerance {%f} for context '%s'"), *What, *Expected.ToString(), *Actual.ToString(), Tolerance, ContextObject ? *ContextObject->GetName() : TEXT("")));
+		return false;
+	}
+	else
+	{
+		LogStep(ELogVerbosity::Log, FString::Printf(TEXT("Quat assertion passed (%s)"), *What));
+		return true;
+	}
+}
+
+bool AFunctionalTest::AssertNotEqual_Quat(const FQuat& Actual, const FQuat& NotExpected, const FString& What, const UObject* ContextObject)
+{
+	if (NotExpected.Equals(Actual))
+	{
+		LogStep(ELogVerbosity::Error, FString::Printf(TEXT("Expected '%s' not to be {%s} for context '%s'"), *What, *NotExpected.ToString(), ContextObject ? *ContextObject->GetName() : TEXT("")));
+		return false;
+	}
+	else
+	{
+		LogStep(ELogVerbosity::Log, FString::Printf(TEXT("Quat assertion passed (%s)"), *What));
+		return true;
+	}
+}
+
+bool AFunctionalTest::AssertEqual_Matrix(const FMatrix& Actual, const FMatrix& Expected, const FString& What, const float Tolerance, const UObject* ContextObject)
+{
+	if (!Expected.Equals(Actual, Tolerance))
+	{
+		LogStep(ELogVerbosity::Error, FString::Printf(TEXT("Expected '%s' to be {%s} but it was {%s} within tolerance {%f} for context '%s'"), *What, *Expected.ToString(), *Actual.ToString(), Tolerance, ContextObject ? *ContextObject->GetName() : TEXT("")));
+		return false;
+	}
+	else
+	{
+		LogStep(ELogVerbosity::Log, FString::Printf(TEXT("Matrix assertion passed (%s)"), *What));
+		return true;
+	}
+}
+
+bool AFunctionalTest::AssertNotEqual_Matrix(const FMatrix& Actual, const FMatrix& NotExpected, const FString& What, const UObject* ContextObject)
+{
+	if (NotExpected.Equals(Actual))
+	{
+		LogStep(ELogVerbosity::Error, FString::Printf(TEXT("Expected '%s' not to be {%s} for context '%s'"), *What, *NotExpected.ToString(), ContextObject ? *ContextObject->GetName() : TEXT("")));
+		return false;
+	}
+	else
+	{
+		LogStep(ELogVerbosity::Log, FString::Printf(TEXT("Matrix assertion passed (%s)"), *What));
+		return true;
+	}
+}
+
+bool AFunctionalTest::AssertEqual_String(const FString& Actual, const FString& Expected, const FString& What, const UObject* ContextObject)
 {
 	if ( !Expected.Equals(Actual) )
 	{
@@ -896,7 +1037,7 @@ bool AFunctionalTest::AssertEqual_String(const FString Actual, const FString Exp
 	}
 }
 
-bool AFunctionalTest::AssertNotEqual_String(const FString Actual, const FString NotExpected, const FString& What, const UObject* ContextObject)
+bool AFunctionalTest::AssertNotEqual_String(const FString& Actual, const FString& NotExpected, const FString& What, const UObject* ContextObject)
 {
 	if ( NotExpected.Equals(Actual) )
 	{
