Expose the number of frames of delay before taking a screenshot

The code had a hardcoded 5 frame delay. All this patch does is expose a setting to control this value (with the default being 5 for backwards compatibility).

The logic is unchanged, first the time based delay is checked, then the frame based delay is checked. You can set either delay to 0 to only use one of them.

This new feature will be used in path tracing related tests to ensure proper convergence regardless of the speed of the machine running the test and without having to set overly lengthy time delays which slow down faster machines without providing guarantees on slower machines.

I have verified that the feature works as expected on all the path tracing tests, including the ones that were frequently failing due to being on the threshold of running too slowly to accumulate enough samples by the time delay.

diff --git a/Engine/Source/Developer/FunctionalTesting/Private/ScreenshotFunctionalTestBase.cpp b/Engine/Source/Developer/FunctionalTesting/Private/ScreenshotFunctionalTestBase.cpp
--- a/Engine/Source/Developer/FunctionalTesting/Private/ScreenshotFunctionalTestBase.cpp
+++ b/Engine/Source/Developer/FunctionalTesting/Private/ScreenshotFunctionalTestBase.cpp
@@ -54,7 +54,7 @@ bool AScreenshotFunctionalTestBase::IsReady_Implementation()
 {
 	if ((GetWorld()->GetTimeSeconds() - RunTime) > ScreenshotOptions.Delay)
 	{
-		return (GFrameNumber - RunFrame) > 5;
+		return (GFrameNumber - RunFrame) > ScreenshotOptions.FrameDelay;
 	}
 
 	return false;
diff --git a/Engine/Source/Developer/FunctionalTesting/Public/AutomationScreenshotOptions.h b/Engine/Source/Developer/FunctionalTesting/Public/AutomationScreenshotOptions.h
--- a/Engine/Source/Developer/FunctionalTesting/Public/AutomationScreenshotOptions.h
+++ b/Engine/Source/Developer/FunctionalTesting/Public/AutomationScreenshotOptions.h
@@ -74,6 +74,7 @@ public:
 	FAutomationScreenshotOptions()
 		: Resolution(ForceInit)
 		, Delay(0.2f)
+		, FrameDelay(5)
 		, bOverride_OverrideTimeTo(false)
 		, OverrideTimeTo(0.0f)
 		, bDisableNoisyRenderingFeatures(true)
@@ -92,6 +93,7 @@ public:
 	FAutomationScreenshotOptions(EComparisonTolerance InTolerance)
 		: Resolution(ForceInit)
 		, Delay(0.2f)
+		, FrameDelay(5)
 		, bOverride_OverrideTimeTo(false)
 		, OverrideTimeTo(0.0f)
 		, bDisableNoisyRenderingFeatures(true)
@@ -116,11 +118,17 @@ public:
 	FVector2D Resolution;
 
 	/**
-	 * The delay before we take the screenshot.
+	 * The delay before we take the screenshot (measured in seconds). Both this delay and the frame delay must be met before the screenshot is taken.
 	 */
 	UPROPERTY(EditAnywhere, BlueprintReadWrite, Category="Screenshot")
 	float Delay;
 
+	/**
+	 * The delay before we take the screenshot (measured in number of frames). Both this frame delay and the time delay must be met before the screenshot is taken.
+	 */
+	UPROPERTY(EditAnywhere, BlueprintReadWrite, Category = "Screenshot")
+	uint32 FrameDelay;
+
 	UPROPERTY(EditAnywhere, BlueprintReadWrite, Category = Overrides, meta = (PinHiddenByDefault, InlineEditConditionToggle))
 	bool bOverride_OverrideTimeTo;
 
