Changed FrameDelay to an int32 which is allowed by blueprints.  C# UHT wasn't generating an error.

diff --git a/Engine/Source/Developer/FunctionalTesting/Private/ScreenshotFunctionalTestBase.cpp b/Engine/Source/Developer/FunctionalTesting/Private/ScreenshotFunctionalTestBase.cpp
--- a/Engine/Source/Developer/FunctionalTesting/Private/ScreenshotFunctionalTestBase.cpp
+++ b/Engine/Source/Developer/FunctionalTesting/Private/ScreenshotFunctionalTestBase.cpp
@@ -54,7 +54,7 @@ bool AScreenshotFunctionalTestBase::IsReady_Implementation()
 {
 	if ((GetWorld()->GetTimeSeconds() - RunTime) > ScreenshotOptions.Delay)
 	{
-		return (GFrameNumber - RunFrame) > ScreenshotOptions.FrameDelay;
+		return int32(GFrameNumber - RunFrame) > ScreenshotOptions.FrameDelay;
 	}
 
 	return false;
diff --git a/Engine/Source/Developer/FunctionalTesting/Public/AutomationScreenshotOptions.h b/Engine/Source/Developer/FunctionalTesting/Public/AutomationScreenshotOptions.h
--- a/Engine/Source/Developer/FunctionalTesting/Public/AutomationScreenshotOptions.h
+++ b/Engine/Source/Developer/FunctionalTesting/Public/AutomationScreenshotOptions.h
@@ -127,7 +127,7 @@ public:
 	 * The delay before we take the screenshot (measured in number of frames). Both this frame delay and the time delay must be met before the screenshot is taken.
 	 */
 	UPROPERTY(EditAnywhere, BlueprintReadWrite, Category = "Screenshot")
-	uint32 FrameDelay;
+	int32 FrameDelay;
 
 	UPROPERTY(EditAnywhere, BlueprintReadWrite, Category = Overrides, meta = (PinHiddenByDefault, InlineEditConditionToggle))
 	bool bOverride_OverrideTimeTo;
