Updated ../Engine/Source/Developer/... to inline gen.cpp files

Before:
Total CPU Time: 53783.640625 s
Total time in Parallel executor: 558.66 seconds
After:
Total CPU Time: 47886.140625 s
Total time in Parallel executor: 498.81 seconds

diff --git a/Engine/Source/Developer/AITestSuite/Private/Actions/TestPawnAction_CallFunction.cpp b/Engine/Source/Developer/AITestSuite/Private/Actions/TestPawnAction_CallFunction.cpp
--- a/Engine/Source/Developer/AITestSuite/Private/Actions/TestPawnAction_CallFunction.cpp
+++ b/Engine/Source/Developer/AITestSuite/Private/Actions/TestPawnAction_CallFunction.cpp
@@ -3,6 +3,8 @@
 #include "Actions/TestPawnAction_CallFunction.h"
 #include "Engine/World.h"
 
+#include UE_INLINE_GENERATED_CPP_BY_NAME(TestPawnAction_CallFunction)
+
 UTestPawnAction_CallFunction::UTestPawnAction_CallFunction(const FObjectInitializer& ObjectInitializer)
 	: Super(ObjectInitializer)
 	, FunctionToCall(nullptr)
@@ -63,3 +65,4 @@ void UTestPawnAction_CallFunction::OnChildFinished(UPawnAction& Action, EPawnAct
 	Super::OnChildFinished(Action, WithResult);
 	(*FunctionToCall)(*GetOwnerComponent(), *this, ETestPawnActionMessage::ChildFinished);
 }
+
diff --git a/Engine/Source/Developer/AITestSuite/Private/Actions/TestPawnAction_Log.cpp b/Engine/Source/Developer/AITestSuite/Private/Actions/TestPawnAction_Log.cpp
--- a/Engine/Source/Developer/AITestSuite/Private/Actions/TestPawnAction_Log.cpp
+++ b/Engine/Source/Developer/AITestSuite/Private/Actions/TestPawnAction_Log.cpp
@@ -4,6 +4,8 @@
 #include "TestLogger.h"
 #include "Engine/World.h"
 
+#include UE_INLINE_GENERATED_CPP_BY_NAME(TestPawnAction_Log)
+
 UTestPawnAction_Log::UTestPawnAction_Log(const FObjectInitializer& ObjectInitializer) 
 	: Super(ObjectInitializer)
 	, Logger(nullptr)
@@ -54,3 +56,4 @@ void UTestPawnAction_Log::OnChildFinished(UPawnAction& Action, EPawnActionResult
 	Super::OnChildFinished(Action, WithResult);
 	Logger->Log(ETestPawnActionMessage::ChildFinished);
 }
+
diff --git a/Engine/Source/Developer/AITestSuite/Private/BehaviorTree/TestBTDecorator_Blackboard.cpp b/Engine/Source/Developer/AITestSuite/Private/BehaviorTree/TestBTDecorator_Blackboard.cpp
--- a/Engine/Source/Developer/AITestSuite/Private/BehaviorTree/TestBTDecorator_Blackboard.cpp
+++ b/Engine/Source/Developer/AITestSuite/Private/BehaviorTree/TestBTDecorator_Blackboard.cpp
@@ -3,6 +3,8 @@
 #include "BehaviorTree/TestBTDecorator_Blackboard.h"
 #include "MockAI_BT.h"
 
+#include UE_INLINE_GENERATED_CPP_BY_NAME(TestBTDecorator_Blackboard)
+
 UTestBTDecorator_Blackboard::UTestBTDecorator_Blackboard(const FObjectInitializer& ObjectInitializer) : Super(ObjectInitializer)
 , LogIndexBecomeRelevant(-1)
 , LogIndexCeaseRelevant(-1)
@@ -35,3 +37,4 @@ void UTestBTDecorator_Blackboard::LogExecution(int32 LogNumber) const
 		UMockAI_BT::ExecutionLog.Add(LogNumber);
 	}
 }
+
diff --git a/Engine/Source/Developer/AITestSuite/Private/BehaviorTree/TestBTDecorator_CantExecute.cpp b/Engine/Source/Developer/AITestSuite/Private/BehaviorTree/TestBTDecorator_CantExecute.cpp
--- a/Engine/Source/Developer/AITestSuite/Private/BehaviorTree/TestBTDecorator_CantExecute.cpp
+++ b/Engine/Source/Developer/AITestSuite/Private/BehaviorTree/TestBTDecorator_CantExecute.cpp
@@ -2,6 +2,8 @@
 
 #include "BehaviorTree/TestBTDecorator_CantExecute.h"
 
+#include UE_INLINE_GENERATED_CPP_BY_NAME(TestBTDecorator_CantExecute)
+
 UTestBTDecorator_CantExecute::UTestBTDecorator_CantExecute(const FObjectInitializer& ObjectInitializer) : Super(ObjectInitializer)
 {
 	NodeName = TEXT("Can't Exexcute");
@@ -15,3 +17,4 @@ bool UTestBTDecorator_CantExecute::CalculateRawConditionValue(UBehaviorTreeCompo
 {
 	return false;
 }
+
diff --git a/Engine/Source/Developer/AITestSuite/Private/BehaviorTree/TestBTDecorator_DelayedAbort.cpp b/Engine/Source/Developer/AITestSuite/Private/BehaviorTree/TestBTDecorator_DelayedAbort.cpp
--- a/Engine/Source/Developer/AITestSuite/Private/BehaviorTree/TestBTDecorator_DelayedAbort.cpp
+++ b/Engine/Source/Developer/AITestSuite/Private/BehaviorTree/TestBTDecorator_DelayedAbort.cpp
@@ -3,6 +3,8 @@
 #include "BehaviorTree/TestBTDecorator_DelayedAbort.h"
 #include "AITestsCommon.h"
 
+#include UE_INLINE_GENERATED_CPP_BY_NAME(TestBTDecorator_DelayedAbort)
+
 UTestBTDecorator_DelayedAbort::UTestBTDecorator_DelayedAbort(const FObjectInitializer& ObjectInitializer) : Super(ObjectInitializer)
 {
 	NodeName = "Delayed Abort";
@@ -38,3 +40,4 @@ uint16 UTestBTDecorator_DelayedAbort::GetInstanceMemorySize() const
 {
 	return sizeof(FBTDelayedAbortMemory);
 }
+
diff --git a/Engine/Source/Developer/AITestSuite/Private/BehaviorTree/TestBTService_Log.cpp b/Engine/Source/Developer/AITestSuite/Private/BehaviorTree/TestBTService_Log.cpp
--- a/Engine/Source/Developer/AITestSuite/Private/BehaviorTree/TestBTService_Log.cpp
+++ b/Engine/Source/Developer/AITestSuite/Private/BehaviorTree/TestBTService_Log.cpp
@@ -4,6 +4,8 @@
 #include "BehaviorTree/BlackboardComponent.h"
 #include "MockAI_BT.h"
 
+#include UE_INLINE_GENERATED_CPP_BY_NAME(TestBTService_Log)
+
 UTestBTService_Log::UTestBTService_Log(const FObjectInitializer& ObjectInitializer) : Super(ObjectInitializer)
 {
 	NodeName = "LogService";
@@ -73,3 +75,4 @@ void UTestBTService_Log::SetFlagOnTick(FName InKeyNameTick, bool bInCallTickOnSe
 	KeyNameTick = InKeyNameTick;
 	bCallTickOnSearchStart = bInCallTickOnSearchStart;
 }
+
diff --git a/Engine/Source/Developer/AITestSuite/Private/BehaviorTree/TestBTService_StopTree.cpp b/Engine/Source/Developer/AITestSuite/Private/BehaviorTree/TestBTService_StopTree.cpp
--- a/Engine/Source/Developer/AITestSuite/Private/BehaviorTree/TestBTService_StopTree.cpp
+++ b/Engine/Source/Developer/AITestSuite/Private/BehaviorTree/TestBTService_StopTree.cpp
@@ -3,6 +3,8 @@
 #include "BehaviorTree/TestBTService_StopTree.h"
 #include "MockAI_BT.h"
 
+#include UE_INLINE_GENERATED_CPP_BY_NAME(TestBTService_StopTree)
+
 UTestBTService_StopTree::UTestBTService_StopTree(const FObjectInitializer& ObjectInitializer) : Super(ObjectInitializer)
 {
 	NodeName = "StopTreeService";
@@ -48,4 +50,4 @@ void UTestBTService_StopTree::TickNode(UBehaviorTreeComponent& OwnerComp, uint8*
 		UMockAI_BT::ExecutionLog.Add(LogIndex);
 		OwnerComp.StopTree();
 	}
-}
\ No newline at end of file
+}
diff --git a/Engine/Source/Developer/AITestSuite/Private/BehaviorTree/TestBTTask_LatentWithFlags.cpp b/Engine/Source/Developer/AITestSuite/Private/BehaviorTree/TestBTTask_LatentWithFlags.cpp
--- a/Engine/Source/Developer/AITestSuite/Private/BehaviorTree/TestBTTask_LatentWithFlags.cpp
+++ b/Engine/Source/Developer/AITestSuite/Private/BehaviorTree/TestBTTask_LatentWithFlags.cpp
@@ -4,6 +4,8 @@
 #include "BehaviorTree/BlackboardComponent.h"
 #include "MockAI_BT.h"
 
+#include UE_INLINE_GENERATED_CPP_BY_NAME(TestBTTask_LatentWithFlags)
+
 UTestBTTask_LatentWithFlags::UTestBTTask_LatentWithFlags()
 {
 	INIT_TASK_NODE_NOTIFY_FLAGS();
@@ -91,4 +93,4 @@ void UTestBTTask_LatentWithFlags::LogExecution(UBehaviorTreeComponent& OwnerComp
 	{
 		UMockAI_BT::ExecutionLog.Add(LogNumber);
 	}
-}
\ No newline at end of file
+}
diff --git a/Engine/Source/Developer/AITestSuite/Private/BehaviorTree/TestBTTask_Log.cpp b/Engine/Source/Developer/AITestSuite/Private/BehaviorTree/TestBTTask_Log.cpp
--- a/Engine/Source/Developer/AITestSuite/Private/BehaviorTree/TestBTTask_Log.cpp
+++ b/Engine/Source/Developer/AITestSuite/Private/BehaviorTree/TestBTTask_Log.cpp
@@ -3,6 +3,8 @@
 #include "BehaviorTree/TestBTTask_Log.h"
 #include "MockAI_BT.h"
 
+#include UE_INLINE_GENERATED_CPP_BY_NAME(TestBTTask_Log)
+
 UTestBTTask_Log::UTestBTTask_Log(const FObjectInitializer& ObjectInitializer) : Super(ObjectInitializer)
 {
 	NodeName = "Log";
@@ -57,3 +59,4 @@ void UTestBTTask_Log::LogExecution(UBehaviorTreeComponent& OwnerComp, int32 LogN
 		UMockAI_BT::ExecutionLog.Add(LogNumber);
 	}
 }
+
diff --git a/Engine/Source/Developer/AITestSuite/Private/BehaviorTree/TestBTTask_SetFlag.cpp b/Engine/Source/Developer/AITestSuite/Private/BehaviorTree/TestBTTask_SetFlag.cpp
--- a/Engine/Source/Developer/AITestSuite/Private/BehaviorTree/TestBTTask_SetFlag.cpp
+++ b/Engine/Source/Developer/AITestSuite/Private/BehaviorTree/TestBTTask_SetFlag.cpp
@@ -4,6 +4,8 @@
 #include "BehaviorTree/Blackboard/BlackboardKeyType_Bool.h"
 #include "BehaviorTree/BlackboardComponent.h"
 
+#include UE_INLINE_GENERATED_CPP_BY_NAME(TestBTTask_SetFlag)
+
 UTestBTTask_SetFlag::UTestBTTask_SetFlag(const FObjectInitializer& ObjectInitializer) : Super(ObjectInitializer)
 {
 	NodeName = "Log";
@@ -27,4 +29,4 @@ EBTNodeResult::Type UTestBTTask_SetFlag::AbortTask(UBehaviorTreeComponent& Owner
 		OwnerComp.GetBlackboardComponent()->SetValue<UBlackboardKeyType_Bool>(OnAbortKeyName, bOnAbortValue);
 	}
 	return EBTNodeResult::Aborted;
-}
\ No newline at end of file
+}
diff --git a/Engine/Source/Developer/AITestSuite/Private/BehaviorTree/TestBTTask_SetValue.cpp b/Engine/Source/Developer/AITestSuite/Private/BehaviorTree/TestBTTask_SetValue.cpp
--- a/Engine/Source/Developer/AITestSuite/Private/BehaviorTree/TestBTTask_SetValue.cpp
+++ b/Engine/Source/Developer/AITestSuite/Private/BehaviorTree/TestBTTask_SetValue.cpp
@@ -4,6 +4,8 @@
 #include "BehaviorTree/Blackboard/BlackboardKeyType_Int.h"
 #include "BehaviorTree/BlackboardComponent.h"
 
+#include UE_INLINE_GENERATED_CPP_BY_NAME(TestBTTask_SetValue)
+
 UTestBTTask_SetValue::UTestBTTask_SetValue(const FObjectInitializer& ObjectInitializer) : Super(ObjectInitializer)
 {
 	NodeName = "SetValue";
@@ -28,3 +30,4 @@ EBTNodeResult::Type UTestBTTask_SetValue::AbortTask(UBehaviorTreeComponent& Owne
 	}
 	return EBTNodeResult::Aborted;	
 }
+
diff --git a/Engine/Source/Developer/AITestSuite/Private/BehaviorTree/TestBTTask_StopTree.cpp b/Engine/Source/Developer/AITestSuite/Private/BehaviorTree/TestBTTask_StopTree.cpp
--- a/Engine/Source/Developer/AITestSuite/Private/BehaviorTree/TestBTTask_StopTree.cpp
+++ b/Engine/Source/Developer/AITestSuite/Private/BehaviorTree/TestBTTask_StopTree.cpp
@@ -3,6 +3,8 @@
 #include "BehaviorTree/TestBTTask_StopTree.h"
 #include "MockAI_BT.h"
 
+#include UE_INLINE_GENERATED_CPP_BY_NAME(TestBTTask_StopTree)
+
 UTestBTTask_StopTree::UTestBTTask_StopTree(const FObjectInitializer& ObjectInitializer) : Super(ObjectInitializer)
 {
 	StopTimming = EBTTestTaskStopTree::DuringExecute;
@@ -51,3 +53,4 @@ void UTestBTTask_StopTree::OnTaskFinished(UBehaviorTreeComponent& OwnerComp, uin
 		OwnerComp.StopTree();
 	}
 }
+
diff --git a/Engine/Source/Developer/AITestSuite/Private/BehaviorTree/TestBTTask_ToggleFlag.cpp b/Engine/Source/Developer/AITestSuite/Private/BehaviorTree/TestBTTask_ToggleFlag.cpp
--- a/Engine/Source/Developer/AITestSuite/Private/BehaviorTree/TestBTTask_ToggleFlag.cpp
+++ b/Engine/Source/Developer/AITestSuite/Private/BehaviorTree/TestBTTask_ToggleFlag.cpp
@@ -4,6 +4,8 @@
 #include "BehaviorTree/Blackboard/BlackboardKeyType_Bool.h"
 #include "BehaviorTree/BlackboardComponent.h"
 
+#include UE_INLINE_GENERATED_CPP_BY_NAME(TestBTTask_ToggleFlag)
+
 UTestBTTask_ToggleFlag::UTestBTTask_ToggleFlag(const FObjectInitializer& ObjectInitializer) : Super(ObjectInitializer)
 {
 	NodeName = "Log";
@@ -19,4 +21,4 @@ EBTNodeResult::Type UTestBTTask_ToggleFlag::ExecuteTask(UBehaviorTreeComponent&
 		OwnerComp.GetBlackboardComponent()->SetValue<UBlackboardKeyType_Bool>(KeyName, !OwnerComp.GetBlackboardComponent()->GetValue<UBlackboardKeyType_Bool>(KeyName));
 	}
 	return TaskResult;
-}
\ No newline at end of file
+}
diff --git a/Engine/Source/Developer/AITestSuite/Private/MockAI/MockAI.cpp b/Engine/Source/Developer/AITestSuite/Private/MockAI/MockAI.cpp
--- a/Engine/Source/Developer/AITestSuite/Private/MockAI/MockAI.cpp
+++ b/Engine/Source/Developer/AITestSuite/Private/MockAI/MockAI.cpp
@@ -6,6 +6,8 @@
 #include "Actions/PawnActionsComponent.h"
 #include "BrainComponent.h"
 
+#include UE_INLINE_GENERATED_CPP_BY_NAME(MockAI)
+
 //----------------------------------------------------------------------//
 // 
 //----------------------------------------------------------------------//
@@ -102,3 +104,4 @@ void UMockAI::TickMe(float DeltaTime)
 	}
 }
 
+
diff --git a/Engine/Source/Developer/AITestSuite/Private/MockAI/MockAI_BT.cpp b/Engine/Source/Developer/AITestSuite/Private/MockAI/MockAI_BT.cpp
--- a/Engine/Source/Developer/AITestSuite/Private/MockAI/MockAI_BT.cpp
+++ b/Engine/Source/Developer/AITestSuite/Private/MockAI/MockAI_BT.cpp
@@ -5,6 +5,8 @@
 #include "BehaviorTree/BehaviorTreeComponent.h"
 #include "BehaviorTree/BehaviorTree.h"
 
+#include UE_INLINE_GENERATED_CPP_BY_NAME(MockAI_BT)
+
 //----------------------------------------------------------------------//
 // 
 //----------------------------------------------------------------------//
@@ -43,3 +45,4 @@ void UMockAI_BT::RunBT(UBehaviorTree& BTAsset, EBTExecutionMode::Type RunType)
 
 	BTComp->StartTree(BTAsset, RunType);
 }
+
diff --git a/Engine/Source/Developer/AITestSuite/Private/MockAI/MockGameplayTasks.cpp b/Engine/Source/Developer/AITestSuite/Private/MockAI/MockGameplayTasks.cpp
--- a/Engine/Source/Developer/AITestSuite/Private/MockAI/MockGameplayTasks.cpp
+++ b/Engine/Source/Developer/AITestSuite/Private/MockAI/MockGameplayTasks.cpp
@@ -3,6 +3,8 @@
 #include "MockGameplayTasks.h"
 #include "TestLogger.h"
 
+#include UE_INLINE_GENERATED_CPP_BY_NAME(MockGameplayTasks)
+
 //----------------------------------------------------------------------//
 // 
 //----------------------------------------------------------------------//
@@ -79,3 +81,4 @@ void UMockTask_Log::ExternalCancel()
 //----------------------------------------------------------------------//
 // 
 //----------------------------------------------------------------------//
+
diff --git a/Engine/Source/Developer/Android/AndroidPlatformEditor/Private/AndroidSDKSettings.cpp b/Engine/Source/Developer/Android/AndroidPlatformEditor/Private/AndroidSDKSettings.cpp
--- a/Engine/Source/Developer/Android/AndroidPlatformEditor/Private/AndroidSDKSettings.cpp
+++ b/Engine/Source/Developer/Android/AndroidPlatformEditor/Private/AndroidSDKSettings.cpp
@@ -7,6 +7,8 @@
 //#include "EngineTypes.h"
 #include "Interfaces/ITargetPlatformManagerModule.h"
 
+#include UE_INLINE_GENERATED_CPP_BY_NAME(AndroidSDKSettings)
+
 DEFINE_LOG_CATEGORY_STATIC(AndroidSDKSettings, Log, All);
 
 UAndroidSDKSettings::UAndroidSDKSettings(const FObjectInitializer& ObjectInitializer)
@@ -68,3 +70,4 @@ void UAndroidSDKSettings::UpdateTargetModulePaths()
 }
 
 #endif
+
diff --git a/Engine/Source/Developer/AnimationDataController/Private/AnimDataController.cpp b/Engine/Source/Developer/AnimationDataController/Private/AnimDataController.cpp
--- a/Engine/Source/Developer/AnimationDataController/Private/AnimDataController.cpp
+++ b/Engine/Source/Developer/AnimationDataController/Private/AnimDataController.cpp
@@ -15,6 +15,8 @@
 #include "Animation/AnimSequenceHelpers.h"
 #include "Math/UnrealMathUtility.h"
 
+#include UE_INLINE_GENERATED_CPP_BY_NAME(AnimDataController)
+
 #define LOCTEXT_NAMESPACE "AnimDataController"
 
 #if WITH_EDITOR
@@ -2138,3 +2140,4 @@ void UAnimDataController::PopulateWithExistingModel(TScriptInterface<IAnimationD
 
 #undef LOCTEXT_NAMESPACE // "AnimDataController"
 
+
diff --git a/Engine/Source/Developer/AssetTools/Private/AdvancedCopyCustomization.cpp b/Engine/Source/Developer/AssetTools/Private/AdvancedCopyCustomization.cpp
--- a/Engine/Source/Developer/AssetTools/Private/AdvancedCopyCustomization.cpp
+++ b/Engine/Source/Developer/AssetTools/Private/AdvancedCopyCustomization.cpp
@@ -8,6 +8,8 @@
 #include "Engine/Level.h"
 #include "Engine/MapBuildDataRegistry.h"
 
+#include UE_INLINE_GENERATED_CPP_BY_NAME(AdvancedCopyCustomization)
+
 
 #define LOCTEXT_NAMESPACE "AdvancedCopyCustomization"
 
@@ -56,3 +58,4 @@ void UAdvancedCopyCustomization::SetPackageThatInitiatedCopy(const FString& InBa
 }
 
 #undef LOCTEXT_NAMESPACE
+
diff --git a/Engine/Source/Developer/AssetTools/Private/AssetTools.cpp b/Engine/Source/Developer/AssetTools/Private/AssetTools.cpp
--- a/Engine/Source/Developer/AssetTools/Private/AssetTools.cpp
+++ b/Engine/Source/Developer/AssetTools/Private/AssetTools.cpp
@@ -167,6 +167,8 @@
 #include "Settings/EditorExperimentalSettings.h"
 #endif
 
+#include UE_INLINE_GENERATED_CPP_BY_NAME(AssetTools)
+
 #define LOCTEXT_NAMESPACE "AssetTools"
 
 TScriptInterface<IAssetTools> UAssetToolsHelpers::GetAssetTools()
@@ -4317,3 +4319,4 @@ void UAssetToolsImpl::UnregisterIsNameAllowedDelegate(const FName OwnerName)
 }
 
 #undef LOCTEXT_NAMESPACE
+
diff --git a/Engine/Source/Developer/AutomationController/Private/AutomationControllerSettings.cpp b/Engine/Source/Developer/AutomationController/Private/AutomationControllerSettings.cpp
--- a/Engine/Source/Developer/AutomationController/Private/AutomationControllerSettings.cpp
+++ b/Engine/Source/Developer/AutomationController/Private/AutomationControllerSettings.cpp
@@ -3,6 +3,8 @@
 
 #include "AutomationControllerSettings.h"
 
+#include UE_INLINE_GENERATED_CPP_BY_NAME(AutomationControllerSettings)
+
 DEFINE_LOG_CATEGORY_STATIC(LogAutomationControllerSettings, Log, All)
 
 UAutomationControllerSettings::UAutomationControllerSettings(const FObjectInitializer& ObjectInitializer)
@@ -26,4 +28,4 @@ void UAutomationControllerSettings::PostInitProperties()
 		UE_LOG(LogAutomationControllerSettings, Warning, TEXT("UAutomationControllerSettings::bTreatLogWarningsAsTestErrors is deprecated. Use bElevateLogWarningsToErrors instead."));
 		bElevateLogWarningsToErrors = false;
 	}
-}
\ No newline at end of file
+}
diff --git a/Engine/Source/Developer/AutomationController/Private/AutomationDeviceClusterManager.cpp b/Engine/Source/Developer/AutomationController/Private/AutomationDeviceClusterManager.cpp
--- a/Engine/Source/Developer/AutomationController/Private/AutomationDeviceClusterManager.cpp
+++ b/Engine/Source/Developer/AutomationController/Private/AutomationDeviceClusterManager.cpp
@@ -3,6 +3,8 @@
 #include "AutomationDeviceClusterManager.h"
 #include "IAutomationControllerManager.h"
 
+#include UE_INLINE_GENERATED_CPP_BY_NAME(AutomationDeviceClusterManager)
+
 void FAutomationDeviceClusterManager::Reset()
 {
 	Clusters.Empty();
@@ -333,3 +335,4 @@ bool FAutomationDeviceClusterManager::HasActiveDevice()
 	}
 	return IsDeviceAvailable;
 }
+
diff --git a/Engine/Source/Developer/AutomationWindow/Private/AutomationPresetManager.cpp b/Engine/Source/Developer/AutomationWindow/Private/AutomationPresetManager.cpp
--- a/Engine/Source/Developer/AutomationWindow/Private/AutomationPresetManager.cpp
+++ b/Engine/Source/Developer/AutomationWindow/Private/AutomationPresetManager.cpp
@@ -7,6 +7,8 @@
 #include "Serialization/JsonSerializer.h"
 #include "JsonObjectConverter.h"
 
+#include UE_INLINE_GENERATED_CPP_BY_NAME(AutomationPresetManager)
+
 FAutomationTestPresetManager::FAutomationTestPresetManager()
 {
 	// Add the None Option
@@ -124,3 +126,4 @@ void FAutomationTestPresetManager::LoadPresets()
 		}
 	}
 }
+
diff --git a/Engine/Source/Developer/DeveloperToolSettings/Private/CookerSettings.cpp b/Engine/Source/Developer/DeveloperToolSettings/Private/CookerSettings.cpp
--- a/Engine/Source/Developer/DeveloperToolSettings/Private/CookerSettings.cpp
+++ b/Engine/Source/Developer/DeveloperToolSettings/Private/CookerSettings.cpp
@@ -3,6 +3,8 @@
 #include "CookerSettings.h"
 #include "UObject/UnrealType.h"
 
+#include UE_INLINE_GENERATED_CPP_BY_NAME(CookerSettings)
+
 DEFINE_LOG_CATEGORY_STATIC(LogCookerSettings, Log, All);
 
 UCookerSettings::UCookerSettings(const FObjectInitializer& ObjectInitializer)
@@ -84,4 +86,4 @@ void UCookerSettings::PostEditChangeProperty(FPropertyChangedEvent& PropertyChan
 	}
 }
 
-#endif
\ No newline at end of file
+#endif
diff --git a/Engine/Source/Developer/DeveloperToolSettings/Private/ProjectPackagingSettings.cpp b/Engine/Source/Developer/DeveloperToolSettings/Private/ProjectPackagingSettings.cpp
--- a/Engine/Source/Developer/DeveloperToolSettings/Private/ProjectPackagingSettings.cpp
+++ b/Engine/Source/Developer/DeveloperToolSettings/Private/ProjectPackagingSettings.cpp
@@ -8,6 +8,8 @@
 #include "DeveloperToolSettingsDelegates.h"
 #include "InstalledPlatformInfo.h"
 
+#include UE_INLINE_GENERATED_CPP_BY_NAME(ProjectPackagingSettings)
+
 #define LOCTEXT_NAMESPACE "SettingsClasses"
 
 
@@ -303,3 +305,4 @@ void UProjectPackagingSettings::SetBuildTargetForPlatform(FName PlatformName, FS
 }
 
 #undef LOCTEXT_NAMESPACE
+
diff --git a/Engine/Source/Developer/DrawPrimitiveDebugger/Private/DrawPrimitiveDebuggerConfig.cpp b/Engine/Source/Developer/DrawPrimitiveDebugger/Private/DrawPrimitiveDebuggerConfig.cpp
--- a/Engine/Source/Developer/DrawPrimitiveDebugger/Private/DrawPrimitiveDebuggerConfig.cpp
+++ b/Engine/Source/Developer/DrawPrimitiveDebugger/Private/DrawPrimitiveDebuggerConfig.cpp
@@ -1,8 +1,10 @@
-﻿// Copyright Epic Games, Inc. All Rights Reserved.
+// Copyright Epic Games, Inc. All Rights Reserved.
 
 #include "DrawPrimitiveDebuggerConfig.h"
 #include "Misc/ConfigCacheIni.h"
 
+#include UE_INLINE_GENERATED_CPP_BY_NAME(DrawPrimitiveDebuggerConfig)
+
 void UDrawPrimitiveDebuggerUserSettings::SetFontSize(const int32 InFontSize)
 {
 	UDrawPrimitiveDebuggerUserSettings* Settings = GetMutableDefault<UDrawPrimitiveDebuggerUserSettings>();
@@ -12,3 +14,4 @@ void UDrawPrimitiveDebuggerUserSettings::SetFontSize(const int32 InFontSize)
 	GConfig->SetInt(TEXT("/Script/DrawPrimitiveDebugger.DrawPrimitiveDebuggerUserSettings"), TEXT("FontSize"), Settings->FontSize, *ConfigFileName);
 	GConfig->Flush(false);
 }
+
diff --git a/Engine/Source/Developer/FunctionalTesting/Private/AutomationBlueprintFunctionLibrary.cpp b/Engine/Source/Developer/FunctionalTesting/Private/AutomationBlueprintFunctionLibrary.cpp
--- a/Engine/Source/Developer/FunctionalTesting/Private/AutomationBlueprintFunctionLibrary.cpp
+++ b/Engine/Source/Developer/FunctionalTesting/Private/AutomationBlueprintFunctionLibrary.cpp
@@ -58,6 +58,8 @@
 #endif
 #include "FunctionalTestBase.h"
 
+#include UE_INLINE_GENERATED_CPP_BY_NAME(AutomationBlueprintFunctionLibrary)
+
 
 #define LOCTEXT_NAMESPACE "Automation"
 
@@ -1444,3 +1446,4 @@ void UAutomationBlueprintFunctionLibrary::SetEditorViewportVisualizeBuffer( FNam
 }
 
 #undef LOCTEXT_NAMESPACE
+
diff --git a/Engine/Source/Developer/FunctionalTesting/Private/FuncTestRenderingComponent.cpp b/Engine/Source/Developer/FunctionalTesting/Private/FuncTestRenderingComponent.cpp
--- a/Engine/Source/Developer/FunctionalTesting/Private/FuncTestRenderingComponent.cpp
+++ b/Engine/Source/Developer/FunctionalTesting/Private/FuncTestRenderingComponent.cpp
@@ -8,6 +8,8 @@
 //#include "PrimitiveViewRelevance.h"
 #include "SceneManagement.h"
 
+#include UE_INLINE_GENERATED_CPP_BY_NAME(FuncTestRenderingComponent)
+
 //----------------------------------------------------------------------//
 // FFTestRenderingSceneProxy
 //----------------------------------------------------------------------//
@@ -122,3 +124,4 @@ FBoxSphereBounds UFuncTestRenderingComponent::CalcBounds(const FTransform& Local
 
 	return FBoxSphereBounds(ForceInitToZero);
 }
+
diff --git a/Engine/Source/Developer/FunctionalTesting/Private/FunctionalAITest.cpp b/Engine/Source/Developer/FunctionalTesting/Private/FunctionalAITest.cpp
--- a/Engine/Source/Developer/FunctionalTesting/Private/FunctionalAITest.cpp
+++ b/Engine/Source/Developer/FunctionalTesting/Private/FunctionalAITest.cpp
@@ -11,6 +11,8 @@
 #include "NavMesh/RecastNavMesh.h"
 #include "NavigationOctree.h"
 
+#include UE_INLINE_GENERATED_CPP_BY_NAME(FunctionalAITest)
+
 AFunctionalAITestBase::AFunctionalAITestBase( const FObjectInitializer& ObjectInitializer )
 	: Super(ObjectInitializer)
 	, CurrentSpawnSetIndex(INDEX_NONE)
@@ -531,3 +533,4 @@ bool AFunctionalAITest::IsValidSpawnSetIndex(const int32 Index) const
 {
 	return SpawnSets.IsValidIndex(Index);
 }
+
diff --git a/Engine/Source/Developer/FunctionalTesting/Private/FunctionalTest.cpp b/Engine/Source/Developer/FunctionalTesting/Private/FunctionalTest.cpp
--- a/Engine/Source/Developer/FunctionalTesting/Private/FunctionalTest.cpp
+++ b/Engine/Source/Developer/FunctionalTesting/Private/FunctionalTest.cpp
@@ -26,6 +26,8 @@
 #include "Misc/RuntimeErrors.h"
 #include "FunctionalTestBase.h"
 
+#include UE_INLINE_GENERATED_CPP_BY_NAME(FunctionalTest)
+
 DECLARE_CYCLE_STAT(TEXT("FunctionalTest - RunTest"), STAT_FunctionalTest_RunTest, STATGROUP_FunctionalTest);
 DECLARE_CYCLE_STAT(TEXT("FunctionalTest - StartTest"), STAT_FunctionalTest_StartTest, STATGROUP_FunctionalTest);
 DECLARE_CYCLE_STAT(TEXT("FunctionalTest - PrepareTest"), STAT_FunctionalTest_PrepareTest, STATGROUP_FunctionalTest);
@@ -1649,3 +1651,4 @@ void UAutomationPerformaceHelper::EndStatsFile()
 		GEngine->Exec(World, TEXT("Stat StopFile"));
 	}
 }
+
diff --git a/Engine/Source/Developer/FunctionalTesting/Private/FunctionalTestGameMode.cpp b/Engine/Source/Developer/FunctionalTesting/Private/FunctionalTestGameMode.cpp
--- a/Engine/Source/Developer/FunctionalTesting/Private/FunctionalTestGameMode.cpp
+++ b/Engine/Source/Developer/FunctionalTesting/Private/FunctionalTestGameMode.cpp
@@ -3,8 +3,11 @@
 #include "FunctionalTestGameMode.h"
 #include "GameFramework/SpectatorPawn.h"
 
+#include UE_INLINE_GENERATED_CPP_BY_NAME(FunctionalTestGameMode)
+
 AFunctionalTestGameMode::AFunctionalTestGameMode(const FObjectInitializer& ObjectInitializer)
 	: Super(ObjectInitializer)
 {
 	DefaultPawnClass = ASpectatorPawn::StaticClass();
 }
+
diff --git a/Engine/Source/Developer/FunctionalTesting/Private/FunctionalTestLevelScript.cpp b/Engine/Source/Developer/FunctionalTesting/Private/FunctionalTestLevelScript.cpp
--- a/Engine/Source/Developer/FunctionalTesting/Private/FunctionalTestLevelScript.cpp
+++ b/Engine/Source/Developer/FunctionalTesting/Private/FunctionalTestLevelScript.cpp
@@ -2,7 +2,10 @@
 
 #include "FunctionalTestLevelScript.h"
 
+#include UE_INLINE_GENERATED_CPP_BY_NAME(FunctionalTestLevelScript)
+
 AFunctionalTestLevelScript::AFunctionalTestLevelScript( const FObjectInitializer& ObjectInitializer )
 	: Super(ObjectInitializer)
 {
 }
+
diff --git a/Engine/Source/Developer/FunctionalTesting/Private/FunctionalTestUtilityLibrary.cpp b/Engine/Source/Developer/FunctionalTesting/Private/FunctionalTestUtilityLibrary.cpp
--- a/Engine/Source/Developer/FunctionalTesting/Private/FunctionalTestUtilityLibrary.cpp
+++ b/Engine/Source/Developer/FunctionalTesting/Private/FunctionalTestUtilityLibrary.cpp
@@ -3,6 +3,8 @@
 #include "FunctionalTestUtilityLibrary.h"
 #include "TraceQueryTestResults.h"
 
+#include UE_INLINE_GENERATED_CPP_BY_NAME(FunctionalTestUtilityLibrary)
+
 UTraceQueryTestResults* UFunctionalTestUtilityLibrary::TraceChannelTestUtil(UObject* WorldContextObject, const FTraceChannelTestBatchOptions& BatchOptions, const FVector Start, const FVector End, float SphereCapsuleRadius, float CapsuleHalfHeight, FVector BoxHalfSize, const FRotator Orientation, ETraceTypeQuery TraceChannel, TArray<TEnumAsByte<EObjectTypeQuery> > ObjectTypes, FName ProfileName, bool bTraceComplex, const TArray<AActor*>& ActorsToIgnore, bool bIgnoreSelf, EDrawDebugTrace::Type DrawDebugType, FLinearColor TraceColor, FLinearColor TraceHitColor, float DrawTime)
 {
 	UTraceQueryTestResults* Results = NewObject<UTraceQueryTestResults>();
@@ -102,4 +104,4 @@ UTraceQueryTestResults* UFunctionalTestUtilityLibrary::TraceChannelTestUtil(UObj
 	}
 
 	return Results;
-}
\ No newline at end of file
+}
diff --git a/Engine/Source/Developer/FunctionalTesting/Private/FunctionalTestingManager.cpp b/Engine/Source/Developer/FunctionalTesting/Private/FunctionalTestingManager.cpp
--- a/Engine/Source/Developer/FunctionalTesting/Private/FunctionalTestingManager.cpp
+++ b/Engine/Source/Developer/FunctionalTesting/Private/FunctionalTestingManager.cpp
@@ -10,6 +10,8 @@
 #include "Engine/Engine.h"
 #include "Misc/RuntimeErrors.h"
 
+#include UE_INLINE_GENERATED_CPP_BY_NAME(FunctionalTestingManager)
+
 #if WITH_EDITOR
 
 //----------------------------------------------------------------------//
@@ -423,3 +425,4 @@ void UFunctionalTestingManager::SetReproString(FString ReproString)
 		ReproString.ParseIntoArray(TestReproStrings, FFunctionalTesting::ReproStringTestSeparator, /*InCullEmpty=*/true);
 	}
 }
+
diff --git a/Engine/Source/Developer/FunctionalTesting/Private/FunctionalUIScreenshotTest.cpp b/Engine/Source/Developer/FunctionalTesting/Private/FunctionalUIScreenshotTest.cpp
--- a/Engine/Source/Developer/FunctionalTesting/Private/FunctionalUIScreenshotTest.cpp
+++ b/Engine/Source/Developer/FunctionalTesting/Private/FunctionalUIScreenshotTest.cpp
@@ -15,6 +15,8 @@
 #include "Slate/SceneViewport.h"
 #include "Slate/WidgetRenderer.h"
 
+#include UE_INLINE_GENERATED_CPP_BY_NAME(FunctionalUIScreenshotTest)
+
 AFunctionalUIScreenshotTest::AFunctionalUIScreenshotTest( const FObjectInitializer& ObjectInitializer )
 	: AScreenshotFunctionalTestBase(ObjectInitializer)
 {
@@ -226,3 +228,4 @@ void AFunctionalUIScreenshotTest::RequestScreenshot()
 	check(OutColorData.Num() == ScreenshotSize.X * ScreenshotSize.Y);
 	GameViewportClient->OnScreenshotCaptured().Broadcast(ScreenshotSize.X, ScreenshotSize.Y, OutColorData);
 }
+
diff --git a/Engine/Source/Developer/FunctionalTesting/Private/GroundTruthData.cpp b/Engine/Source/Developer/FunctionalTesting/Private/GroundTruthData.cpp
--- a/Engine/Source/Developer/FunctionalTesting/Private/GroundTruthData.cpp
+++ b/Engine/Source/Developer/FunctionalTesting/Private/GroundTruthData.cpp
@@ -15,6 +15,8 @@
 #include "SourceControlHelpers.h"
 #endif
 
+#include UE_INLINE_GENERATED_CPP_BY_NAME(GroundTruthData)
+
 DEFINE_LOG_CATEGORY_STATIC(GroundTruthLog, Log, Log)
 
 UGroundTruthData::UGroundTruthData()
@@ -108,4 +110,4 @@ void UGroundTruthData::PostEditChangeProperty(FPropertyChangedEvent& PropertyCha
 	}
 }
 
-#endif
\ No newline at end of file
+#endif
diff --git a/Engine/Source/Developer/FunctionalTesting/Private/ScreenshotFunctionalTest.cpp b/Engine/Source/Developer/FunctionalTesting/Private/ScreenshotFunctionalTest.cpp
--- a/Engine/Source/Developer/FunctionalTesting/Private/ScreenshotFunctionalTest.cpp
+++ b/Engine/Source/Developer/FunctionalTesting/Private/ScreenshotFunctionalTest.cpp
@@ -15,6 +15,8 @@
 #include "Slate/SceneViewport.h"
 #include "UObject/AutomationObjectVersion.h"
 
+#include UE_INLINE_GENERATED_CPP_BY_NAME(ScreenshotFunctionalTest)
+
 AScreenshotFunctionalTest::AScreenshotFunctionalTest( const FObjectInitializer& ObjectInitializer )
 	: AScreenshotFunctionalTestBase(ObjectInitializer)
 	, bCameraCutOnScreenshotPrep(true)
@@ -90,3 +92,4 @@ void AScreenshotFunctionalTest::RequestScreenshot()
 		FScreenshotRequest::RequestScreenshot(bShowUI);
 	}
 }
+
diff --git a/Engine/Source/Developer/FunctionalTesting/Private/ScreenshotFunctionalTestBase.cpp b/Engine/Source/Developer/FunctionalTesting/Private/ScreenshotFunctionalTestBase.cpp
--- a/Engine/Source/Developer/FunctionalTesting/Private/ScreenshotFunctionalTestBase.cpp
+++ b/Engine/Source/Developer/FunctionalTesting/Private/ScreenshotFunctionalTestBase.cpp
@@ -19,6 +19,8 @@
 #include "Logging/LogMacros.h"
 #include "UObject/AutomationObjectVersion.h"
 
+#include UE_INLINE_GENERATED_CPP_BY_NAME(ScreenshotFunctionalTestBase)
+
 #define	WITH_EDITOR_AUTOMATION_TESTS	(WITH_EDITOR && WITH_AUTOMATION_TESTS)
 
 DEFINE_LOG_CATEGORY_STATIC(LogScreenshotFunctionalTest, Log, Log)
@@ -254,3 +256,4 @@ void AScreenshotFunctionalTestBase::Serialize(FArchive& Ar)
 	}
 }
 
+
diff --git a/Engine/Source/Developer/FunctionalTesting/Private/TestPhaseComponent.cpp b/Engine/Source/Developer/FunctionalTesting/Private/TestPhaseComponent.cpp
--- a/Engine/Source/Developer/FunctionalTesting/Private/TestPhaseComponent.cpp
+++ b/Engine/Source/Developer/FunctionalTesting/Private/TestPhaseComponent.cpp
@@ -2,7 +2,10 @@
 
 #include "TestPhaseComponent.h"
 
+#include UE_INLINE_GENERATED_CPP_BY_NAME(TestPhaseComponent)
+
 UTestPhaseComponent::UTestPhaseComponent(const FObjectInitializer& ObjectInitializer)
 	: Super(ObjectInitializer)
 {
 }
+
diff --git a/Engine/Source/Developer/FunctionalTesting/Private/TraceQueryTestResults.cpp b/Engine/Source/Developer/FunctionalTesting/Private/TraceQueryTestResults.cpp
--- a/Engine/Source/Developer/FunctionalTesting/Private/TraceQueryTestResults.cpp
+++ b/Engine/Source/Developer/FunctionalTesting/Private/TraceQueryTestResults.cpp
@@ -5,6 +5,8 @@
 #include "Components/PrimitiveComponent.h"
 #include "PhysicalMaterials/PhysicalMaterial.h"
 
+#include UE_INLINE_GENERATED_CPP_BY_NAME(TraceQueryTestResults)
+
 FString FTraceQueryTestNames::ToString() const
 {
 	return FString::Printf(TEXT("Component:%s, Actor:%s, PhysicalMaterial:%s"), *ComponentName.ToString(), *ActorName.ToString(), *PhysicalMaterialName.ToString());
@@ -228,3 +230,4 @@ bool UTraceQueryTestResults::AssertEqual(const UTraceQueryTestResults* Expected,
 	}
 }
 
+
diff --git a/Engine/Source/Developer/GameplayDebugger/Private/GameplayDebuggerCategoryReplicator.cpp b/Engine/Source/Developer/GameplayDebugger/Private/GameplayDebuggerCategoryReplicator.cpp
--- a/Engine/Source/Developer/GameplayDebugger/Private/GameplayDebuggerCategoryReplicator.cpp
+++ b/Engine/Source/Developer/GameplayDebugger/Private/GameplayDebuggerCategoryReplicator.cpp
@@ -10,6 +10,8 @@
 #include "Net/UnrealNetwork.h"
 #include "VisualLogger/VisualLogger.h"
 
+#include UE_INLINE_GENERATED_CPP_BY_NAME(GameplayDebuggerCategoryReplicator)
+
 #if WITH_EDITOR
 #include "Editor.h"
 #include "Engine/Selection.h"
@@ -971,3 +973,4 @@ bool AGameplayDebuggerCategoryReplicator::IsCategoryEnabled(int32 CategoryId) co
 {
 	return Categories.IsValidIndex(CategoryId) && Categories[CategoryId]->IsCategoryEnabled(); 
 }
+
diff --git a/Engine/Source/Developer/GameplayDebugger/Private/GameplayDebuggerConfig.cpp b/Engine/Source/Developer/GameplayDebugger/Private/GameplayDebuggerConfig.cpp
--- a/Engine/Source/Developer/GameplayDebugger/Private/GameplayDebuggerConfig.cpp
+++ b/Engine/Source/Developer/GameplayDebugger/Private/GameplayDebuggerConfig.cpp
@@ -6,6 +6,8 @@
 #include "CoreGlobals.h"
 #include "Misc/ConfigCacheIni.h"
 
+#include UE_INLINE_GENERATED_CPP_BY_NAME(GameplayDebuggerConfig)
+
 
 UGameplayDebuggerConfig::UGameplayDebuggerConfig(const FObjectInitializer& ObjectInitializer) : Super(ObjectInitializer)
 {
@@ -297,3 +299,4 @@ void UGameplayDebuggerUserSettings::SetFontSize(const int32 InFontSize)
 	GConfig->SetInt(TEXT("/Script/GameplayDebugger.GameplayDebuggerUserSettings"), TEXT("FontSize"), Setting->FontSize, *ConfigFileName);
 	GConfig->Flush(false);
 }
+
diff --git a/Engine/Source/Developer/GameplayDebugger/Private/GameplayDebuggerLocalController.cpp b/Engine/Source/Developer/GameplayDebugger/Private/GameplayDebuggerLocalController.cpp
--- a/Engine/Source/Developer/GameplayDebugger/Private/GameplayDebuggerLocalController.cpp
+++ b/Engine/Source/Developer/GameplayDebugger/Private/GameplayDebuggerLocalController.cpp
@@ -25,6 +25,8 @@
 #include "EngineUtils.h"
 #include "HAL/IConsoleManager.h"
 
+#include UE_INLINE_GENERATED_CPP_BY_NAME(GameplayDebuggerLocalController)
+
 #if WITH_EDITOR
 #include "Editor/GameplayDebuggerEdMode.h"
 #include "EditorModeManager.h"
@@ -1081,3 +1083,4 @@ FAutoConsoleCommandWithWorldAndArgs FGameplayDebuggerConsoleCommands::SetFontSiz
 	TEXT("Configures gameplay debugger's font size. Usage: gdt.fontsize <fontSize> (default = 10)"),
 	FConsoleCommandWithWorldAndArgsDelegate::CreateStatic(&FGameplayDebuggerConsoleCommands::SetFontSize)
 );
+
diff --git a/Engine/Source/Developer/GameplayDebugger/Private/GameplayDebuggerPlayerManager.cpp b/Engine/Source/Developer/GameplayDebugger/Private/GameplayDebuggerPlayerManager.cpp
--- a/Engine/Source/Developer/GameplayDebugger/Private/GameplayDebuggerPlayerManager.cpp
+++ b/Engine/Source/Developer/GameplayDebugger/Private/GameplayDebuggerPlayerManager.cpp
@@ -13,6 +13,8 @@
 #include "Net/UnrealNetwork.h"
 #include "GameFramework/GameModeBase.h"
 
+#include UE_INLINE_GENERATED_CPP_BY_NAME(GameplayDebuggerPlayerManager)
+
 AGameplayDebuggerPlayerManager::AGameplayDebuggerPlayerManager(const FObjectInitializer& ObjectInitializer) : Super(ObjectInitializer)
 {
 	PrimaryActorTick.bCanEverTick = true;
@@ -355,4 +357,4 @@ void AGameplayDebuggerPlayerManager::OnReplayScrubTeardown(UWorld* InWorld)
 	{
 		UpdateAuthReplicators();
 	}
-}
\ No newline at end of file
+}
diff --git a/Engine/Source/Developer/GameplayDebugger/Private/GameplayDebuggerRenderingComponent.cpp b/Engine/Source/Developer/GameplayDebugger/Private/GameplayDebuggerRenderingComponent.cpp
--- a/Engine/Source/Developer/GameplayDebugger/Private/GameplayDebuggerRenderingComponent.cpp
+++ b/Engine/Source/Developer/GameplayDebugger/Private/GameplayDebuggerRenderingComponent.cpp
@@ -4,6 +4,8 @@
 #include "GameplayDebuggerCategoryReplicator.h"
 #include "GameplayDebuggerCategory.h"
 
+#include UE_INLINE_GENERATED_CPP_BY_NAME(GameplayDebuggerRenderingComponent)
+
 //////////////////////////////////////////////////////////////////////////
 // FGameplayDebuggerCompositeSceneProxy
 
@@ -169,3 +171,4 @@ FBoxSphereBounds UGameplayDebuggerRenderingComponent::CalcBounds(const FTransfor
 {
 	return FBoxSphereBounds(FBox::BuildAABB(FVector::ZeroVector, FVector(1000000.0f, 1000000.0f, 1000000.0f)));
 }
+
diff --git a/Engine/Source/Developer/Localization/Private/LocTextHelper.cpp b/Engine/Source/Developer/Localization/Private/LocTextHelper.cpp
--- a/Engine/Source/Developer/Localization/Private/LocTextHelper.cpp
+++ b/Engine/Source/Developer/Localization/Private/LocTextHelper.cpp
@@ -14,6 +14,8 @@
 #include "Serialization/JsonInternationalizationArchiveSerializer.h"
 #include "Serialization/JsonInternationalizationMetadataSerializer.h"
 
+#include UE_INLINE_GENERATED_CPP_BY_NAME(LocTextHelper)
+
 #define LOCTEXT_NAMESPACE "LocTextHelper"
 
 DEFINE_LOG_CATEGORY_STATIC(LogLocTextHelper, Log, All);
@@ -1915,3 +1917,4 @@ TSharedPtr<FArchiveEntry> FLocTextHelper::FindTranslationImpl(const FString& InC
 }
 
 #undef LOCTEXT_NAMESPACE
+
diff --git a/Engine/Source/Developer/Localization/Private/LocalizationSettings.cpp b/Engine/Source/Developer/Localization/Private/LocalizationSettings.cpp
--- a/Engine/Source/Developer/Localization/Private/LocalizationSettings.cpp
+++ b/Engine/Source/Developer/Localization/Private/LocalizationSettings.cpp
@@ -5,6 +5,8 @@
 #include "ISourceControlProvider.h"
 #include "ISourceControlModule.h"
 
+#include UE_INLINE_GENERATED_CPP_BY_NAME(LocalizationSettings)
+
 ULocalizationSettings::ULocalizationSettings(const FObjectInitializer& ObjectInitializer)
 	: Super(ObjectInitializer)
 	, EngineTargetSet(ObjectInitializer.CreateDefaultSubobject<ULocalizationTargetSet>(this, TEXT("EngineLocalizationTargetSet")))
@@ -138,3 +140,4 @@ void FLocalizationSourceControlSettings::SetSourceControlAutoSubmitEnabled(const
 {
 	GConfig->SetBool(*LocalizationSourceControlSettingsCategoryName, *SourceControlAutoSubmitEnabledSettingName, bIsEnabled, GEditorPerProjectIni);
 }
+
diff --git a/Engine/Source/Developer/Localization/Private/LocalizationTargetTypes.cpp b/Engine/Source/Developer/Localization/Private/LocalizationTargetTypes.cpp
--- a/Engine/Source/Developer/Localization/Private/LocalizationTargetTypes.cpp
+++ b/Engine/Source/Developer/Localization/Private/LocalizationTargetTypes.cpp
@@ -14,6 +14,8 @@
 #include "LocalizationConfigurationScript.h"
 #include "LocalizationSettings.h"
 
+#include UE_INLINE_GENERATED_CPP_BY_NAME(LocalizationTargetTypes)
+
 #define LOCTEXT_NAMESPACE "LocalizationTargetTypes"
 
 FString FLocalizationGatherPathRootUtil::GetResolvedPathRootToken(const ELocalizationGatherPathRoot InPathRoot)
@@ -730,3 +732,4 @@ void ULocalizationTargetSet::PostEditChangeProperty(FPropertyChangedEvent& Prope
 #endif
 
 #undef LOCTEXT_NAMESPACE
+
diff --git a/Engine/Source/Developer/Localization/Private/PortableObjectPipeline.cpp b/Engine/Source/Developer/Localization/Private/PortableObjectPipeline.cpp
--- a/Engine/Source/Developer/Localization/Private/PortableObjectPipeline.cpp
+++ b/Engine/Source/Developer/Localization/Private/PortableObjectPipeline.cpp
@@ -9,6 +9,8 @@
 #include "Internationalization/TextNamespaceUtil.h"
 #include "LocTextHelper.h"
 
+#include UE_INLINE_GENERATED_CPP_BY_NAME(PortableObjectPipeline)
+
 DEFINE_LOG_CATEGORY_STATIC(LogPortableObjectPipeline, Log, All);
 
 namespace
@@ -739,3 +741,4 @@ void PortableObjectPipeline::ParseBasicPOFileEntry(const FPortableObjectEntry& P
 		OutTranslation = ConditionPOStringForArchive(POEntry.MsgStr[0]);
 	}
 }
+
diff --git a/Engine/Source/Developer/LogVisualizer/Private/LogVisualizerSessionSettings.cpp b/Engine/Source/Developer/LogVisualizer/Private/LogVisualizerSessionSettings.cpp
--- a/Engine/Source/Developer/LogVisualizer/Private/LogVisualizerSessionSettings.cpp
+++ b/Engine/Source/Developer/LogVisualizer/Private/LogVisualizerSessionSettings.cpp
@@ -1,6 +1,9 @@
 // Copyright Epic Games, Inc. All Rights Reserved.
 
 #include "LogVisualizerSessionSettings.h"
+
+#include UE_INLINE_GENERATED_CPP_BY_NAME(LogVisualizerSessionSettings)
+
 #if WITH_EDITOR
 #include "UnrealEdMisc.h"
 #endif // WITH_EDITOR
@@ -25,3 +28,4 @@ void ULogVisualizerSessionSettings::PostEditChangeProperty(struct FPropertyChang
 	SettingChangedEvent.Broadcast(Name);
 }
 #endif
+
diff --git a/Engine/Source/Developer/LogVisualizer/Private/LogVisualizerSettings.cpp b/Engine/Source/Developer/LogVisualizer/Private/LogVisualizerSettings.cpp
--- a/Engine/Source/Developer/LogVisualizer/Private/LogVisualizerSettings.cpp
+++ b/Engine/Source/Developer/LogVisualizer/Private/LogVisualizerSettings.cpp
@@ -5,6 +5,9 @@
 #include "VisualLogger/VisualLoggerTypes.h"
 #include "VisualLoggerDatabase.h"
 #include "VisualLogger/VisualLogger.h"
+
+#include UE_INLINE_GENERATED_CPP_BY_NAME(LogVisualizerSettings)
+
 #if WITH_EDITOR
 #include "UnrealEdMisc.h"
 #endif // WITH_EDITOR
@@ -330,3 +333,4 @@ bool FVisualLoggerFilters::IsGraphDataDisabled(FName GraphName, FName DataName)
 	const FName FullName = *(GraphName.ToString() + TEXT("$") + DataName.ToString());
 	return DisabledGraphDatas.Find(FullName) != INDEX_NONE;
 }
+
diff --git a/Engine/Source/Developer/LogVisualizer/Private/VisualLoggerCameraController.cpp b/Engine/Source/Developer/LogVisualizer/Private/VisualLoggerCameraController.cpp
--- a/Engine/Source/Developer/LogVisualizer/Private/VisualLoggerCameraController.cpp
+++ b/Engine/Source/Developer/LogVisualizer/Private/VisualLoggerCameraController.cpp
@@ -8,6 +8,8 @@
 #include "Engine/Player.h"
 #include "VisualLoggerHUD.h"
 
+#include UE_INLINE_GENERATED_CPP_BY_NAME(VisualLoggerCameraController)
+
 
 //----------------------------------------------------------------------//
 // AVisualLoggerCameraController
@@ -154,3 +156,4 @@ void AVisualLoggerCameraController::Select(FHitResult const& Hit)
 		OnActorSelected.ExecuteIfBound(SelectedActor.Get());
 	}
 }
+
diff --git a/Engine/Source/Developer/LogVisualizer/Private/VisualLoggerHUD.cpp b/Engine/Source/Developer/LogVisualizer/Private/VisualLoggerHUD.cpp
--- a/Engine/Source/Developer/LogVisualizer/Private/VisualLoggerHUD.cpp
+++ b/Engine/Source/Developer/LogVisualizer/Private/VisualLoggerHUD.cpp
@@ -9,6 +9,8 @@
 #include "DrawDebugHelpers.h"
 #include "VisualLoggerCameraController.h"
 
+#include UE_INLINE_GENERATED_CPP_BY_NAME(VisualLoggerHUD)
+
 #define LOCTEXT_NAMESPACE "AVisualLoggerHUD"
 
 //----------------------------------------------------------------------//
@@ -72,3 +74,4 @@ void AVisualLoggerHUD::PostRender()
 	}
 }
 #undef LOCTEXT_NAMESPACE
+
diff --git a/Engine/Source/Developer/LogVisualizer/Private/VisualLoggerRenderingActor.cpp b/Engine/Source/Developer/LogVisualizer/Private/VisualLoggerRenderingActor.cpp
--- a/Engine/Source/Developer/LogVisualizer/Private/VisualLoggerRenderingActor.cpp
+++ b/Engine/Source/Developer/LogVisualizer/Private/VisualLoggerRenderingActor.cpp
@@ -6,6 +6,8 @@
 #include "VisualLoggerDatabase.h"
 #include "LogVisualizerPublic.h"
 
+#include UE_INLINE_GENERATED_CPP_BY_NAME(VisualLoggerRenderingActor)
+
 AVisualLoggerRenderingActor::AVisualLoggerRenderingActor(const FObjectInitializer& ObjectInitializer)
 	: Super(ObjectInitializer)
 {
@@ -212,4 +214,4 @@ void AVisualLoggerRenderingActor::IterateDebugShapes(const TFunction<void(const
 #if VLOG_TEST_DEBUG_RENDERING
 	Callback(TestDebugShapes);
 #endif
-}
\ No newline at end of file
+}
diff --git a/Engine/Source/Developer/LogVisualizer/Private/VisualLoggerRenderingActorBase.cpp b/Engine/Source/Developer/LogVisualizer/Private/VisualLoggerRenderingActorBase.cpp
--- a/Engine/Source/Developer/LogVisualizer/Private/VisualLoggerRenderingActorBase.cpp
+++ b/Engine/Source/Developer/LogVisualizer/Private/VisualLoggerRenderingActorBase.cpp
@@ -8,6 +8,8 @@
 #include "VisualLoggerRenderingComponent.h"
 #include "VisualLogger/VisualLogger.h"
 
+#include UE_INLINE_GENERATED_CPP_BY_NAME(VisualLoggerRenderingActorBase)
+
 namespace FDebugDrawing
 {
 	const FVector NavOffset(0, 0, 15);
@@ -520,3 +522,4 @@ void AVisualLoggerRenderingActorBase::GetDebugShapes(const FVisualLogEntry& InEn
 		}
 	}
 }
+
diff --git a/Engine/Source/Developer/RigVMDeveloper/Private/RigVMCompiler/RigVMAST.cpp b/Engine/Source/Developer/RigVMDeveloper/Private/RigVMCompiler/RigVMAST.cpp
--- a/Engine/Source/Developer/RigVMDeveloper/Private/RigVMCompiler/RigVMAST.cpp
+++ b/Engine/Source/Developer/RigVMDeveloper/Private/RigVMCompiler/RigVMAST.cpp
@@ -23,6 +23,8 @@
 #include "VisualGraphUtils.h"
 #include "UObject/FieldIterator.h"
 
+#include UE_INLINE_GENERATED_CPP_BY_NAME(RigVMAST)
+
 FRigVMExprAST::FRigVMExprAST(EType InType, const FRigVMASTProxy& InProxy)
 	: Name(NAME_None)
 	, Type(InType)
@@ -3384,4 +3386,4 @@ FString FRigVMParserAST::GetLinkAsString(const FRigVMASTLinkDescription& InLink)
 
 	return FString::Printf(TEXT("%s -> %s%s%s"), *SourcePin->GetPinPath(), *TargetPin->GetPinPath(),
 		*(InLink.SegmentPath.IsEmpty() ? EmptyString : PeriodString), *InLink.SegmentPath);
-}
\ No newline at end of file
+}
diff --git a/Engine/Source/Developer/RigVMDeveloper/Private/RigVMCompiler/RigVMCompiler.cpp b/Engine/Source/Developer/RigVMDeveloper/Private/RigVMCompiler/RigVMCompiler.cpp
--- a/Engine/Source/Developer/RigVMDeveloper/Private/RigVMCompiler/RigVMCompiler.cpp
+++ b/Engine/Source/Developer/RigVMDeveloper/Private/RigVMCompiler/RigVMCompiler.cpp
@@ -11,6 +11,8 @@
 #include "Stats/StatsHierarchical.h"
 #include "RigVMTypeUtils.h"
 
+#include UE_INLINE_GENERATED_CPP_BY_NAME(RigVMCompiler)
+
 class FRigVMCompilerImportErrorContext : public FOutputDevice
 {
 public:
@@ -2934,3 +2936,4 @@ void URigVMCompiler::ReportError(const FString& InMessage)
 {
 	Settings.Report(EMessageSeverity::Error, nullptr, InMessage);
 }
+
diff --git a/Engine/Source/Developer/RigVMDeveloper/Private/RigVMModel/Nodes/RigVMAggregateNode.cpp b/Engine/Source/Developer/RigVMDeveloper/Private/RigVMModel/Nodes/RigVMAggregateNode.cpp
--- a/Engine/Source/Developer/RigVMDeveloper/Private/RigVMModel/Nodes/RigVMAggregateNode.cpp
+++ b/Engine/Source/Developer/RigVMDeveloper/Private/RigVMModel/Nodes/RigVMAggregateNode.cpp
@@ -4,6 +4,8 @@
 
 #include "RigVMModel/RigVMController.h"
 
+#include UE_INLINE_GENERATED_CPP_BY_NAME(RigVMAggregateNode)
+
 URigVMAggregateNode::URigVMAggregateNode()
 	: Super()
 	, FirstInnerNodeCache(nullptr)
@@ -271,3 +273,4 @@ FText URigVMAggregateNode::GetToolTipTextForPin(const URigVMPin* InPin) const
 	
 	return Super::GetToolTipTextForPin(InPin);
 }
+
diff --git a/Engine/Source/Developer/RigVMDeveloper/Private/RigVMModel/Nodes/RigVMArrayNode.cpp b/Engine/Source/Developer/RigVMDeveloper/Private/RigVMModel/Nodes/RigVMArrayNode.cpp
--- a/Engine/Source/Developer/RigVMDeveloper/Private/RigVMModel/Nodes/RigVMArrayNode.cpp
+++ b/Engine/Source/Developer/RigVMDeveloper/Private/RigVMModel/Nodes/RigVMArrayNode.cpp
@@ -4,6 +4,8 @@
 #include "RigVMModel/RigVMGraph.h"
 #include "RigVMCore/RigVMStruct.h"
 
+#include UE_INLINE_GENERATED_CPP_BY_NAME(RigVMArrayNode)
+
 const FString URigVMArrayNode::ArrayName = TEXT("Array");
 const FString URigVMArrayNode::NumName = TEXT("Num");
 const FString URigVMArrayNode::IndexName = TEXT("Index");
@@ -847,3 +849,4 @@ UObject* URigVMArrayNode::GetCPPTypeObject() const
 	}
 	return ArrayPin->GetCPPTypeObject();
 }
+
diff --git a/Engine/Source/Developer/RigVMDeveloper/Private/RigVMModel/Nodes/RigVMBranchNode.cpp b/Engine/Source/Developer/RigVMDeveloper/Private/RigVMModel/Nodes/RigVMBranchNode.cpp
--- a/Engine/Source/Developer/RigVMDeveloper/Private/RigVMModel/Nodes/RigVMBranchNode.cpp
+++ b/Engine/Source/Developer/RigVMDeveloper/Private/RigVMModel/Nodes/RigVMBranchNode.cpp
@@ -2,7 +2,10 @@
 
 #include "RigVMModel/Nodes/RigVMBranchNode.h"
 
+#include UE_INLINE_GENERATED_CPP_BY_NAME(RigVMBranchNode)
+
 const FString URigVMBranchNode::BranchName = TEXT("Branch");
 const FString URigVMBranchNode::ConditionName = TEXT("Condition");
 const FString URigVMBranchNode::TrueName = TEXT("True");
 const FString URigVMBranchNode::FalseName = TEXT("False");
+
diff --git a/Engine/Source/Developer/RigVMDeveloper/Private/RigVMModel/Nodes/RigVMCollapseNode.cpp b/Engine/Source/Developer/RigVMDeveloper/Private/RigVMModel/Nodes/RigVMCollapseNode.cpp
--- a/Engine/Source/Developer/RigVMDeveloper/Private/RigVMModel/Nodes/RigVMCollapseNode.cpp
+++ b/Engine/Source/Developer/RigVMDeveloper/Private/RigVMModel/Nodes/RigVMCollapseNode.cpp
@@ -3,6 +3,8 @@
 #include "RigVMModel/Nodes/RigVMCollapseNode.h"
 #include "RigVMModel/RigVMFunctionLibrary.h"
 
+#include UE_INLINE_GENERATED_CPP_BY_NAME(RigVMCollapseNode)
+
 URigVMCollapseNode::URigVMCollapseNode()
 	: URigVMLibraryNode()
 {
@@ -27,4 +29,4 @@ URigVMFunctionLibrary* URigVMCollapseNode::GetLibrary() const
 FString URigVMCollapseNode::GetEditorSubGraphName() const
 {
 	return FString::Printf(TEXT("%s_SubGraph"), *GetName());
-}
\ No newline at end of file
+}
diff --git a/Engine/Source/Developer/RigVMDeveloper/Private/RigVMModel/Nodes/RigVMCommentNode.cpp b/Engine/Source/Developer/RigVMDeveloper/Private/RigVMModel/Nodes/RigVMCommentNode.cpp
--- a/Engine/Source/Developer/RigVMDeveloper/Private/RigVMModel/Nodes/RigVMCommentNode.cpp
+++ b/Engine/Source/Developer/RigVMDeveloper/Private/RigVMModel/Nodes/RigVMCommentNode.cpp
@@ -2,6 +2,8 @@
 
 #include "RigVMModel/Nodes/RigVMCommentNode.h"
 
+#include UE_INLINE_GENERATED_CPP_BY_NAME(RigVMCommentNode)
+
 URigVMCommentNode::URigVMCommentNode()
 {
 	Size = FVector2D(400.f, 300.f);
@@ -29,3 +31,4 @@ bool URigVMCommentNode::GetCommentColorBubble() const
 {
 	return bColorBubble;
 }
+
diff --git a/Engine/Source/Developer/RigVMDeveloper/Private/RigVMModel/Nodes/RigVMDispatchNode.cpp b/Engine/Source/Developer/RigVMDeveloper/Private/RigVMModel/Nodes/RigVMDispatchNode.cpp
--- a/Engine/Source/Developer/RigVMDeveloper/Private/RigVMModel/Nodes/RigVMDispatchNode.cpp
+++ b/Engine/Source/Developer/RigVMDeveloper/Private/RigVMModel/Nodes/RigVMDispatchNode.cpp
@@ -3,6 +3,8 @@
 #include "RigVMModel/Nodes/RigVMDispatchNode.h"
 #include "RigVMUserWorkflowRegistry.h"
 
+#include UE_INLINE_GENERATED_CPP_BY_NAME(RigVMDispatchNode)
+
 FString URigVMDispatchNode::GetNodeTitle() const
 {
 	if (const FRigVMDispatchFactory* Factory = GetFactory())
@@ -180,3 +182,4 @@ FRigVMStructUpgradeInfo URigVMDispatchNode::GetUpgradeInfo() const
 	}
 	return FRigVMStructUpgradeInfo();
 }
+
diff --git a/Engine/Source/Developer/RigVMDeveloper/Private/RigVMModel/Nodes/RigVMEnumNode.cpp b/Engine/Source/Developer/RigVMDeveloper/Private/RigVMModel/Nodes/RigVMEnumNode.cpp
--- a/Engine/Source/Developer/RigVMDeveloper/Private/RigVMModel/Nodes/RigVMEnumNode.cpp
+++ b/Engine/Source/Developer/RigVMDeveloper/Private/RigVMModel/Nodes/RigVMEnumNode.cpp
@@ -2,6 +2,8 @@
 
 #include "RigVMModel/Nodes/RigVMEnumNode.h"
 
+#include UE_INLINE_GENERATED_CPP_BY_NAME(RigVMEnumNode)
+
 const FString URigVMEnumNode::EnumName = TEXT("Enum");
 const FString URigVMEnumNode::EnumValueName = TEXT("EnumValue");
 const FString URigVMEnumNode::EnumIndexName = TEXT("EnumIndex");
@@ -53,3 +55,4 @@ FString URigVMEnumNode::GetDefaultValue(const URigVMPin::FPinOverride& InOverrid
 	}
 	return EnumValuePin->GetDefaultValue(InOverride);
 }
+
diff --git a/Engine/Source/Developer/RigVMDeveloper/Private/RigVMModel/Nodes/RigVMFunctionEntryNode.cpp b/Engine/Source/Developer/RigVMDeveloper/Private/RigVMModel/Nodes/RigVMFunctionEntryNode.cpp
--- a/Engine/Source/Developer/RigVMDeveloper/Private/RigVMModel/Nodes/RigVMFunctionEntryNode.cpp
+++ b/Engine/Source/Developer/RigVMDeveloper/Private/RigVMModel/Nodes/RigVMFunctionEntryNode.cpp
@@ -3,6 +3,8 @@
 #include "RigVMModel/Nodes/RigVMFunctionEntryNode.h"
 #include "RigVMModel/RigVMFunctionLibrary.h"
 
+#include UE_INLINE_GENERATED_CPP_BY_NAME(RigVMFunctionEntryNode)
+
 const FRigVMTemplate* URigVMFunctionEntryNode::GetTemplate() const
 {
 	if (URigVMLibraryNode* LibraryNode = GetTypedOuter<URigVMLibraryNode>())
diff --git a/Engine/Source/Developer/RigVMDeveloper/Private/RigVMModel/Nodes/RigVMFunctionReferenceNode.cpp b/Engine/Source/Developer/RigVMDeveloper/Private/RigVMModel/Nodes/RigVMFunctionReferenceNode.cpp
--- a/Engine/Source/Developer/RigVMDeveloper/Private/RigVMModel/Nodes/RigVMFunctionReferenceNode.cpp
+++ b/Engine/Source/Developer/RigVMDeveloper/Private/RigVMModel/Nodes/RigVMFunctionReferenceNode.cpp
@@ -3,6 +3,8 @@
 #include "RigVMModel/Nodes/RigVMFunctionReferenceNode.h"
 #include "RigVMModel/RigVMFunctionLibrary.h"
 
+#include UE_INLINE_GENERATED_CPP_BY_NAME(RigVMFunctionReferenceNode)
+
 FString URigVMFunctionReferenceNode::GetNodeTitle() const
 {
 	if (URigVMLibraryNode* ReferencedNode = GetReferencedNode())
@@ -245,4 +247,4 @@ URigVMLibraryNode* URigVMFunctionReferenceNode::GetReferencedNode() const
 void URigVMFunctionReferenceNode::SetReferencedNode(URigVMLibraryNode* InReferenceNode)
 {
 	ReferencedNodePtr = InReferenceNode;
-}
\ No newline at end of file
+}
diff --git a/Engine/Source/Developer/RigVMDeveloper/Private/RigVMModel/Nodes/RigVMFunctionReturnNode.cpp b/Engine/Source/Developer/RigVMDeveloper/Private/RigVMModel/Nodes/RigVMFunctionReturnNode.cpp
--- a/Engine/Source/Developer/RigVMDeveloper/Private/RigVMModel/Nodes/RigVMFunctionReturnNode.cpp
+++ b/Engine/Source/Developer/RigVMDeveloper/Private/RigVMModel/Nodes/RigVMFunctionReturnNode.cpp
@@ -3,6 +3,8 @@
 #include "RigVMModel/Nodes/RigVMFunctionReturnNode.h"
 #include "RigVMModel/RigVMFunctionLibrary.h"
 
+#include UE_INLINE_GENERATED_CPP_BY_NAME(RigVMFunctionReturnNode)
+
 const FRigVMTemplate* URigVMFunctionReturnNode::GetTemplate() const
 {
 	if (URigVMLibraryNode* LibraryNode = GetTypedOuter<URigVMLibraryNode>())
diff --git a/Engine/Source/Developer/RigVMDeveloper/Private/RigVMModel/Nodes/RigVMIfNode.cpp b/Engine/Source/Developer/RigVMDeveloper/Private/RigVMModel/Nodes/RigVMIfNode.cpp
--- a/Engine/Source/Developer/RigVMDeveloper/Private/RigVMModel/Nodes/RigVMIfNode.cpp
+++ b/Engine/Source/Developer/RigVMDeveloper/Private/RigVMModel/Nodes/RigVMIfNode.cpp
@@ -2,6 +2,8 @@
 
 #include "RigVMModel/Nodes/RigVMIfNode.h"
 
+#include UE_INLINE_GENERATED_CPP_BY_NAME(RigVMIfNode)
+
 const FString URigVMIfNode::IfName = TEXT("If");
 const FString URigVMIfNode::ConditionName = TEXT("Condition");
 const FString URigVMIfNode::TrueName = TEXT("True");
@@ -72,3 +74,4 @@ const FRigVMTemplate* URigVMIfNode::GetTemplate() const
 	return CachedTemplate;
 }
 
+
diff --git a/Engine/Source/Developer/RigVMDeveloper/Private/RigVMModel/Nodes/RigVMInvokeEntryNode.cpp b/Engine/Source/Developer/RigVMDeveloper/Private/RigVMModel/Nodes/RigVMInvokeEntryNode.cpp
--- a/Engine/Source/Developer/RigVMDeveloper/Private/RigVMModel/Nodes/RigVMInvokeEntryNode.cpp
+++ b/Engine/Source/Developer/RigVMDeveloper/Private/RigVMModel/Nodes/RigVMInvokeEntryNode.cpp
@@ -4,6 +4,8 @@
 
 #include "RigVMModel/RigVMGraph.h"
 
+#include UE_INLINE_GENERATED_CPP_BY_NAME(RigVMInvokeEntryNode)
+
 const FString URigVMInvokeEntryNode::EntryName = TEXT("Entry");
 
 URigVMInvokeEntryNode::URigVMInvokeEntryNode()
@@ -29,3 +31,4 @@ URigVMPin* URigVMInvokeEntryNode::GetEntryNamePin() const
 {
 	return FindPin(EntryName);
 }
+
diff --git a/Engine/Source/Developer/RigVMDeveloper/Private/RigVMModel/Nodes/RigVMLibraryNode.cpp b/Engine/Source/Developer/RigVMDeveloper/Private/RigVMModel/Nodes/RigVMLibraryNode.cpp
--- a/Engine/Source/Developer/RigVMDeveloper/Private/RigVMModel/Nodes/RigVMLibraryNode.cpp
+++ b/Engine/Source/Developer/RigVMDeveloper/Private/RigVMModel/Nodes/RigVMLibraryNode.cpp
@@ -7,6 +7,8 @@
 #include "RigVMModel/Nodes/RigVMCollapseNode.h"
 #include "RigVMModel/Nodes/RigVMFunctionReferenceNode.h"
 
+#include UE_INLINE_GENERATED_CPP_BY_NAME(RigVMLibraryNode)
+
 const TArray<URigVMNode*> URigVMLibraryNode::EmptyNodes;
 const TArray<URigVMLink*> URigVMLibraryNode::EmptyLinks;
 
@@ -187,3 +189,4 @@ TArray<FRigVMExternalVariable> URigVMLibraryNode::GetExternalVariables() const
 	return Variables;
 }
 
+
diff --git a/Engine/Source/Developer/RigVMDeveloper/Private/RigVMModel/Nodes/RigVMParameterNode.cpp b/Engine/Source/Developer/RigVMDeveloper/Private/RigVMModel/Nodes/RigVMParameterNode.cpp
--- a/Engine/Source/Developer/RigVMDeveloper/Private/RigVMModel/Nodes/RigVMParameterNode.cpp
+++ b/Engine/Source/Developer/RigVMDeveloper/Private/RigVMModel/Nodes/RigVMParameterNode.cpp
@@ -2,6 +2,8 @@
 
 #include "RigVMModel/Nodes/RigVMParameterNode.h"
 
+#include UE_INLINE_GENERATED_CPP_BY_NAME(RigVMParameterNode)
+
 const FString URigVMParameterNode::ParameterName = TEXT("Parameter");
 const FString URigVMParameterNode::DefaultName = TEXT("Default");
 const FString URigVMParameterNode::ValueName = TEXT("Value");
@@ -76,3 +78,4 @@ FRigVMGraphParameterDescription URigVMParameterNode::GetParameterDescription() c
 	Parameter.DefaultValue = GetDefaultValue();
 	return Parameter;
 }
+
diff --git a/Engine/Source/Developer/RigVMDeveloper/Private/RigVMModel/Nodes/RigVMRerouteNode.cpp b/Engine/Source/Developer/RigVMDeveloper/Private/RigVMModel/Nodes/RigVMRerouteNode.cpp
--- a/Engine/Source/Developer/RigVMDeveloper/Private/RigVMModel/Nodes/RigVMRerouteNode.cpp
+++ b/Engine/Source/Developer/RigVMDeveloper/Private/RigVMModel/Nodes/RigVMRerouteNode.cpp
@@ -3,6 +3,8 @@
 #include "RigVMModel/Nodes/RigVMRerouteNode.h"
 #include "RigVMModel/RigVMGraph.h"
 
+#include UE_INLINE_GENERATED_CPP_BY_NAME(RigVMRerouteNode)
+
 const FString URigVMRerouteNode::RerouteName = TEXT("Reroute");
 const FString URigVMRerouteNode::ValueName = TEXT("Value");
 
@@ -111,3 +113,4 @@ const FRigVMTemplate* URigVMRerouteNode::GetTemplate() const
 	}
 	return CachedTemplate;
 }
+
diff --git a/Engine/Source/Developer/RigVMDeveloper/Private/RigVMModel/Nodes/RigVMSelectNode.cpp b/Engine/Source/Developer/RigVMDeveloper/Private/RigVMModel/Nodes/RigVMSelectNode.cpp
--- a/Engine/Source/Developer/RigVMDeveloper/Private/RigVMModel/Nodes/RigVMSelectNode.cpp
+++ b/Engine/Source/Developer/RigVMDeveloper/Private/RigVMModel/Nodes/RigVMSelectNode.cpp
@@ -2,6 +2,8 @@
 
 #include "RigVMModel/Nodes/RigVMSelectNode.h"
 
+#include UE_INLINE_GENERATED_CPP_BY_NAME(RigVMSelectNode)
+
 const FString URigVMSelectNode::SelectName = TEXT("Select");
 const FString URigVMSelectNode::IndexName = TEXT("Index");
 const FString URigVMSelectNode::ValueName = TEXT("Values");
@@ -96,4 +98,4 @@ const FRigVMTemplate* URigVMSelectNode::GetTemplate() const
 		SelectNodeTemplate = CachedTemplate = FRigVMRegistry::Get().GetOrAddTemplateFromArguments(*SelectName, Arguments, Delegates);
 	}
 	return CachedTemplate;
-}
\ No newline at end of file
+}
diff --git a/Engine/Source/Developer/RigVMDeveloper/Private/RigVMModel/Nodes/RigVMTemplateNode.cpp b/Engine/Source/Developer/RigVMDeveloper/Private/RigVMModel/Nodes/RigVMTemplateNode.cpp
--- a/Engine/Source/Developer/RigVMDeveloper/Private/RigVMModel/Nodes/RigVMTemplateNode.cpp
+++ b/Engine/Source/Developer/RigVMDeveloper/Private/RigVMModel/Nodes/RigVMTemplateNode.cpp
@@ -5,6 +5,8 @@
 #include "RigVMModel/RigVMController.h"
 #include "UObject/ObjectSaveContext.h"
 
+#include UE_INLINE_GENERATED_CPP_BY_NAME(RigVMTemplateNode)
+
 ////////////////////////////////////////////////////////////////////////////////////////////////
 
 void FRigVMTemplatePreferredType::UpdateStringFromIndex()
@@ -1180,3 +1182,4 @@ void URigVMTemplateNode::InitializeFilteredPermutationsFromTypes(bool bAllowCast
 		}
 	}
 }
+
diff --git a/Engine/Source/Developer/RigVMDeveloper/Private/RigVMModel/Nodes/RigVMUnitNode.cpp b/Engine/Source/Developer/RigVMDeveloper/Private/RigVMModel/Nodes/RigVMUnitNode.cpp
--- a/Engine/Source/Developer/RigVMDeveloper/Private/RigVMModel/Nodes/RigVMUnitNode.cpp
+++ b/Engine/Source/Developer/RigVMDeveloper/Private/RigVMModel/Nodes/RigVMUnitNode.cpp
@@ -7,6 +7,8 @@
 #include "RigVMModel/RigVMControllerActions.h"
 #include "RigVMUserWorkflowRegistry.h"
 
+#include UE_INLINE_GENERATED_CPP_BY_NAME(RigVMUnitNode)
+
 void URigVMUnitNode::PostLoad()
 {
 	Super::PostLoad();
@@ -341,3 +343,4 @@ FRigVMStructUpgradeInfo URigVMUnitNode::GetUpgradeInfo() const
 	}
 	return FRigVMStructUpgradeInfo();
 }
+
diff --git a/Engine/Source/Developer/RigVMDeveloper/Private/RigVMModel/Nodes/RigVMVariableNode.cpp b/Engine/Source/Developer/RigVMDeveloper/Private/RigVMModel/Nodes/RigVMVariableNode.cpp
--- a/Engine/Source/Developer/RigVMDeveloper/Private/RigVMModel/Nodes/RigVMVariableNode.cpp
+++ b/Engine/Source/Developer/RigVMDeveloper/Private/RigVMModel/Nodes/RigVMVariableNode.cpp
@@ -4,6 +4,8 @@
 
 #include "RigVMModel/RigVMGraph.h"
 
+#include UE_INLINE_GENERATED_CPP_BY_NAME(RigVMVariableNode)
+
 const FString URigVMVariableNode::VariableName = TEXT("Variable");
 const FString URigVMVariableNode::ValueName = TEXT("Value");
 
@@ -133,3 +135,4 @@ URigVMPin* URigVMVariableNode::GetValuePin() const
 {
 	return FindPin(ValueName);
 }
+
diff --git a/Engine/Source/Developer/RigVMDeveloper/Private/RigVMModel/RigVMBuildData.cpp b/Engine/Source/Developer/RigVMDeveloper/Private/RigVMModel/RigVMBuildData.cpp
--- a/Engine/Source/Developer/RigVMDeveloper/Private/RigVMModel/RigVMBuildData.cpp
+++ b/Engine/Source/Developer/RigVMDeveloper/Private/RigVMModel/RigVMBuildData.cpp
@@ -2,6 +2,8 @@
 
 #include "RigVMModel/RigVMBuildData.h"
 
+#include UE_INLINE_GENERATED_CPP_BY_NAME(RigVMBuildData)
+
 FRigVMReferenceNodeData::FRigVMReferenceNodeData(URigVMFunctionReferenceNode* InReferenceNode)
 {
 	check(InReferenceNode);
@@ -270,3 +272,4 @@ void URigVMBuildData::ClearInvalidReferences()
 }
 
 
+
diff --git a/Engine/Source/Developer/RigVMDeveloper/Private/RigVMModel/RigVMClient.cpp b/Engine/Source/Developer/RigVMDeveloper/Private/RigVMModel/RigVMClient.cpp
--- a/Engine/Source/Developer/RigVMDeveloper/Private/RigVMModel/RigVMClient.cpp
+++ b/Engine/Source/Developer/RigVMDeveloper/Private/RigVMModel/RigVMClient.cpp
@@ -2,6 +2,9 @@
 
 #include "RigVMModel/RigVMClient.h"
 #include "Misc/TransactionObjectEvent.h"
+
+#include UE_INLINE_GENERATED_CPP_BY_NAME(RigVMClient)
+
 #if WITH_EDITOR
 #include "ScopedTransaction.h"
 #endif
@@ -637,3 +640,4 @@ void FRigVMClient::DestroyObject(UObject* InObject)
 		InObject->MarkAsGarbage();
 	}
 }
+
diff --git a/Engine/Source/Developer/RigVMDeveloper/Private/RigVMModel/RigVMController.cpp b/Engine/Source/Developer/RigVMDeveloper/Private/RigVMModel/RigVMController.cpp
--- a/Engine/Source/Developer/RigVMDeveloper/Private/RigVMModel/RigVMController.cpp
+++ b/Engine/Source/Developer/RigVMDeveloper/Private/RigVMModel/RigVMController.cpp
@@ -24,6 +24,8 @@
 #include "RigVMTypeUtils.h"
 #include "Engine/UserDefinedStruct.h"
 
+#include UE_INLINE_GENERATED_CPP_BY_NAME(RigVMController)
+
 #if WITH_EDITOR
 #include "Exporters/Exporter.h"
 #include "UnrealExporter.h"
@@ -19150,3 +19152,4 @@ URigVMControllerSettings::URigVMControllerSettings(const FObjectInitializer& Ini
 {
 	bAutoResolveTemplateNodesWhenLinkingExecute = true;
 }
+
diff --git a/Engine/Source/Developer/RigVMDeveloper/Private/RigVMModel/RigVMControllerActions.cpp b/Engine/Source/Developer/RigVMDeveloper/Private/RigVMModel/RigVMControllerActions.cpp
--- a/Engine/Source/Developer/RigVMDeveloper/Private/RigVMModel/RigVMControllerActions.cpp
+++ b/Engine/Source/Developer/RigVMDeveloper/Private/RigVMModel/RigVMControllerActions.cpp
@@ -3,6 +3,8 @@
 #include "RigVMModel/RigVMControllerActions.h"
 #include "UObject/Package.h"
 
+#include UE_INLINE_GENERATED_CPP_BY_NAME(RigVMControllerActions)
+
 #if WITH_EDITOR
 #include "Misc/TransactionObjectEvent.h"
 #endif
@@ -2412,3 +2414,4 @@ bool FRigVMAddInvokeEntryNodeAction::Redo(URigVMController* InController)
 #endif
 	return false;
 }
+
diff --git a/Engine/Source/Developer/RigVMDeveloper/Private/RigVMModel/RigVMFunctionLibrary.cpp b/Engine/Source/Developer/RigVMDeveloper/Private/RigVMModel/RigVMFunctionLibrary.cpp
--- a/Engine/Source/Developer/RigVMDeveloper/Private/RigVMModel/RigVMFunctionLibrary.cpp
+++ b/Engine/Source/Developer/RigVMDeveloper/Private/RigVMModel/RigVMFunctionLibrary.cpp
@@ -4,6 +4,8 @@
 
 #include "RigVMModel/RigVMController.h"
 
+#include UE_INLINE_GENERATED_CPP_BY_NAME(RigVMFunctionLibrary)
+
 URigVMFunctionLibrary::URigVMFunctionLibrary()
 : URigVMGraph()
 {
@@ -160,3 +162,4 @@ URigVMLibraryNode* URigVMFunctionLibrary::FindPreviouslyLocalizedFunction(URigVM
 	return LocalizedFunction;
 }
 
+
diff --git a/Engine/Source/Developer/RigVMDeveloper/Private/RigVMModel/RigVMGraph.cpp b/Engine/Source/Developer/RigVMDeveloper/Private/RigVMModel/RigVMGraph.cpp
--- a/Engine/Source/Developer/RigVMDeveloper/Private/RigVMModel/RigVMGraph.cpp
+++ b/Engine/Source/Developer/RigVMDeveloper/Private/RigVMModel/RigVMGraph.cpp
@@ -9,6 +9,8 @@
 #include "UObject/Package.h"
 #include "RigVMTypeUtils.h"
 
+#include UE_INLINE_GENERATED_CPP_BY_NAME(RigVMGraph)
+
 URigVMGraph::URigVMGraph()
 : DiagnosticsAST(nullptr)
 , RuntimeAST(nullptr)
@@ -430,3 +432,4 @@ bool URigVMGraph::CanLink(URigVMPin* InSourcePin, URigVMPin* InTargetPin, FStrin
 	return GetDiagnosticsAST()->CanLink(InSourcePin, InTargetPin, OutFailureReason);
 }
 
+
diff --git a/Engine/Source/Developer/RigVMDeveloper/Private/RigVMModel/RigVMLink.cpp b/Engine/Source/Developer/RigVMDeveloper/Private/RigVMModel/RigVMLink.cpp
--- a/Engine/Source/Developer/RigVMDeveloper/Private/RigVMModel/RigVMLink.cpp
+++ b/Engine/Source/Developer/RigVMDeveloper/Private/RigVMModel/RigVMLink.cpp
@@ -3,6 +3,8 @@
 #include "RigVMModel/RigVMLink.h"
 #include "RigVMModel/RigVMGraph.h"
 
+#include UE_INLINE_GENERATED_CPP_BY_NAME(RigVMLink)
+
 URigVMGraph* URigVMLink::GetGraph() const
 {
 	return Cast<URigVMGraph>(GetOuter());
@@ -102,3 +104,4 @@ void URigVMLink::PrepareForCopy()
 		TargetPinPath = CurrenTargetPin->GetPinPath();
 	}
 }
+
diff --git a/Engine/Source/Developer/RigVMDeveloper/Private/RigVMModel/RigVMNode.cpp b/Engine/Source/Developer/RigVMDeveloper/Private/RigVMModel/RigVMNode.cpp
--- a/Engine/Source/Developer/RigVMDeveloper/Private/RigVMModel/RigVMNode.cpp
+++ b/Engine/Source/Developer/RigVMDeveloper/Private/RigVMModel/RigVMNode.cpp
@@ -7,6 +7,8 @@
 #include "RigVMCore/RigVMStruct.h"
 #include "RigVMUserWorkflowRegistry.h"
 
+#include UE_INLINE_GENERATED_CPP_BY_NAME(RigVMNode)
+
 const FString URigVMNode::NodeColorName = TEXT("NodeColor");
 
 #if WITH_EDITOR
@@ -688,4 +690,4 @@ const URigVMNode::FProfilingCache* URigVMNode::UpdateProfilingCacheIfNeeded(URig
 	return Cache.Get();;
 }
 
-#endif
\ No newline at end of file
+#endif
diff --git a/Engine/Source/Developer/RigVMDeveloper/Private/RigVMModel/RigVMPin.cpp b/Engine/Source/Developer/RigVMDeveloper/Private/RigVMModel/RigVMPin.cpp
--- a/Engine/Source/Developer/RigVMDeveloper/Private/RigVMModel/RigVMPin.cpp
+++ b/Engine/Source/Developer/RigVMDeveloper/Private/RigVMModel/RigVMPin.cpp
@@ -22,6 +22,8 @@
 #include "RigVMModel/Nodes/RigVMFunctionReturnNode.h"
 #include "RigVMModel/Nodes/RigVMInvokeEntryNode.h"
 
+#include UE_INLINE_GENERATED_CPP_BY_NAME(RigVMPin)
+
 #if WITH_EDITOR
 #include "UObject/CoreRedirects.h"
 #endif
@@ -1947,3 +1949,4 @@ bool URigVMPin::HasInjectedUnitNodes() const
 	
 	return false;
 }
+
diff --git a/Engine/Source/Developer/RigVMDeveloper/Private/RigVMUserWorkflowRegistry.cpp b/Engine/Source/Developer/RigVMDeveloper/Private/RigVMUserWorkflowRegistry.cpp
--- a/Engine/Source/Developer/RigVMDeveloper/Private/RigVMUserWorkflowRegistry.cpp
+++ b/Engine/Source/Developer/RigVMDeveloper/Private/RigVMUserWorkflowRegistry.cpp
@@ -2,6 +2,8 @@
 
 #include "RigVMUserWorkflowRegistry.h"
 
+#include UE_INLINE_GENERATED_CPP_BY_NAME(RigVMUserWorkflowRegistry)
+
 URigVMUserWorkflowRegistry* URigVMUserWorkflowRegistry::Get()
 {
 	return StaticClass()->GetDefaultObject<URigVMUserWorkflowRegistry>();
@@ -49,3 +51,4 @@ TArray<FRigVMUserWorkflow> URigVMUserWorkflowRegistry::GetWorkflows(ERigVMUserWo
 
 	return Workflows;
 }
+
diff --git a/Engine/Source/Developer/ScreenShotComparisonTools/Private/ImageComparer.cpp b/Engine/Source/Developer/ScreenShotComparisonTools/Private/ImageComparer.cpp
--- a/Engine/Source/Developer/ScreenShotComparisonTools/Private/ImageComparer.cpp
+++ b/Engine/Source/Developer/ScreenShotComparisonTools/Private/ImageComparer.cpp
@@ -11,6 +11,8 @@
 #include "Misc/Paths.h"
 #include "Modules/ModuleManager.h"
 
+#include UE_INLINE_GENERATED_CPP_BY_NAME(ImageComparer)
+
 
 #define LOCTEXT_NAMESPACE "ImageComparer"
 
@@ -571,3 +573,4 @@ double FImageComparer::CompareStructuralSimilarity(const FComparableImage* Image
 }
 
 #undef LOCTEXT_NAMESPACE
+
diff --git a/Engine/Source/Developer/SourceCodeAccess/Private/SourceCodeAccessSettings.cpp b/Engine/Source/Developer/SourceCodeAccess/Private/SourceCodeAccessSettings.cpp
--- a/Engine/Source/Developer/SourceCodeAccess/Private/SourceCodeAccessSettings.cpp
+++ b/Engine/Source/Developer/SourceCodeAccess/Private/SourceCodeAccessSettings.cpp
@@ -3,6 +3,8 @@
 #include "SourceCodeAccessSettings.h"
 #include "Misc/ConfigCacheIni.h"
 
+#include UE_INLINE_GENERATED_CPP_BY_NAME(SourceCodeAccessSettings)
+
 USourceCodeAccessSettings::USourceCodeAccessSettings(const FObjectInitializer& ObjectInitializer)
 	: Super(ObjectInitializer)
 {
@@ -20,3 +22,4 @@ USourceCodeAccessSettings::USourceCodeAccessSettings(const FObjectInitializer& O
 	UE_LOG(LogHAL, Log, TEXT("Linux SourceCodeAccessSettings: %s"), *PreferredAccessor);
 #endif
 }
+
diff --git a/Engine/Source/Developer/SourceControl/Private/SourceControlHelpers.cpp b/Engine/Source/Developer/SourceControl/Private/SourceControlHelpers.cpp
--- a/Engine/Source/Developer/SourceControl/Private/SourceControlHelpers.cpp
+++ b/Engine/Source/Developer/SourceControl/Private/SourceControlHelpers.cpp
@@ -16,6 +16,8 @@
 #include "Logging/MessageLog.h"
 #include "AssetRegistry/AssetRegistryModule.h"
 
+#include UE_INLINE_GENERATED_CPP_BY_NAME(SourceControlHelpers)
+
 #if WITH_EDITOR
 #include "Editor.h"
 #include "PackageTools.h"
@@ -1866,3 +1868,4 @@ ISourceControlProvider& FScopedSourceControl::GetProvider()
 
 
 #undef LOCTEXT_NAMESPACE
+
diff --git a/Engine/Source/Developer/SourceControl/Private/SourceControlPreferences.cpp b/Engine/Source/Developer/SourceControl/Private/SourceControlPreferences.cpp
--- a/Engine/Source/Developer/SourceControl/Private/SourceControlPreferences.cpp
+++ b/Engine/Source/Developer/SourceControl/Private/SourceControlPreferences.cpp
@@ -2,6 +2,8 @@
 
 #include "SourceControlPreferences.h"
 
+#include UE_INLINE_GENERATED_CPP_BY_NAME(SourceControlPreferences)
+
 bool USourceControlPreferences::IsValidationTagEnabled()
 {
 	return GetDefault<USourceControlPreferences>()->bEnableValidationTag;
@@ -16,3 +18,4 @@ bool USourceControlPreferences::AreUncontrolledChangelistsEnabled()
 {
 	return GetDefault<USourceControlPreferences>()->bEnableUncontrolledChangelists;
 }
+
diff --git a/Engine/Source/Developer/ToolMenus/Private/ToolMenu.cpp b/Engine/Source/Developer/ToolMenus/Private/ToolMenu.cpp
--- a/Engine/Source/Developer/ToolMenus/Private/ToolMenu.cpp
+++ b/Engine/Source/Developer/ToolMenus/Private/ToolMenu.cpp
@@ -10,6 +10,8 @@
 #include "Framework/MultiBox/MultiBoxBuilder.h"
 #include "Internationalization/Internationalization.h"
 
+#include UE_INLINE_GENERATED_CPP_BY_NAME(ToolMenu)
+
 UToolMenu::UToolMenu() :
 	MenuType(EMultiBoxType::Menu)
 	, bShouldCleanupContextOnDestroy(true)
@@ -570,3 +572,4 @@ void UToolMenu::Empty()
 	SubMenuParent = nullptr;
 	ModifyBlockWidgetAfterMake.Unbind();
 }
+
diff --git a/Engine/Source/Developer/ToolMenus/Private/ToolMenuContext.cpp b/Engine/Source/Developer/ToolMenus/Private/ToolMenuContext.cpp
--- a/Engine/Source/Developer/ToolMenus/Private/ToolMenuContext.cpp
+++ b/Engine/Source/Developer/ToolMenus/Private/ToolMenuContext.cpp
@@ -10,6 +10,8 @@
 #include "Framework/MultiBox/MultiBoxBuilder.h"
 #include "Internationalization/Internationalization.h"
 
+#include UE_INLINE_GENERATED_CPP_BY_NAME(ToolMenuContext)
+
 
 FToolMenuContext::FToolMenuContext(UObject* InContext)
 {
@@ -188,3 +190,4 @@ void FToolMenuContext::Empty()
 	CommandList.Reset();
 	Extenders.Empty();
 }
+
diff --git a/Engine/Source/Developer/ToolMenus/Private/ToolMenuEntry.cpp b/Engine/Source/Developer/ToolMenus/Private/ToolMenuEntry.cpp
--- a/Engine/Source/Developer/ToolMenus/Private/ToolMenuEntry.cpp
+++ b/Engine/Source/Developer/ToolMenus/Private/ToolMenuEntry.cpp
@@ -11,6 +11,8 @@
 #include "Internationalization/Internationalization.h"
 #include "Input/Events.h"
 
+#include UE_INLINE_GENERATED_CPP_BY_NAME(ToolMenuEntry)
+
 PRAGMA_DISABLE_DEPRECATION_WARNINGS
 FToolMenuEntry::FToolMenuEntry() :
 	Type(EMultiBlockType::None),
@@ -298,3 +300,4 @@ bool FToolMenuEntry::IsScriptObjectDynamicConstruct() const
 	static const FName ConstructMenuEntryName = GET_FUNCTION_NAME_CHECKED(UToolMenuEntryScript, ConstructMenuEntry);
 	return ScriptObject && ScriptObject->GetClass()->IsFunctionImplementedInScript(ConstructMenuEntryName);
 }
+
diff --git a/Engine/Source/Developer/ToolMenus/Private/ToolMenuEntryScript.cpp b/Engine/Source/Developer/ToolMenus/Private/ToolMenuEntryScript.cpp
--- a/Engine/Source/Developer/ToolMenus/Private/ToolMenuEntryScript.cpp
+++ b/Engine/Source/Developer/ToolMenus/Private/ToolMenuEntryScript.cpp
@@ -11,6 +11,8 @@
 #include "Framework/MultiBox/MultiBoxBuilder.h"
 #include "Internationalization/Internationalization.h"
 
+#include UE_INLINE_GENERATED_CPP_BY_NAME(ToolMenuEntryScript)
+
 FScriptSlateIcon::FScriptSlateIcon()
 {
 
@@ -198,3 +200,4 @@ FToolMenuEntryScriptDataAdvanced::FToolMenuEntryScriptDataAdvanced() :
 {
 
 }
+
diff --git a/Engine/Source/Developer/ToolMenus/Private/ToolMenuSection.cpp b/Engine/Source/Developer/ToolMenus/Private/ToolMenuSection.cpp
--- a/Engine/Source/Developer/ToolMenus/Private/ToolMenuSection.cpp
+++ b/Engine/Source/Developer/ToolMenus/Private/ToolMenuSection.cpp
@@ -10,6 +10,8 @@
 #include "Framework/MultiBox/MultiBoxBuilder.h"
 #include "Internationalization/Internationalization.h"
 
+#include UE_INLINE_GENERATED_CPP_BY_NAME(ToolMenuSection)
+
 FToolMenuSection::FToolMenuSection() :
 	ToolMenuSectionDynamic(nullptr),
 	bIsRegistering(false),
@@ -239,3 +241,4 @@ int32 FToolMenuSection::FindBlockInsertIndex(const FToolMenuEntry& InBlock) cons
 
 	return Blocks.Num();
 }
+
diff --git a/Engine/Source/Developer/ToolMenus/Private/ToolMenus.cpp b/Engine/Source/Developer/ToolMenus/Private/ToolMenus.cpp
--- a/Engine/Source/Developer/ToolMenus/Private/ToolMenus.cpp
+++ b/Engine/Source/Developer/ToolMenus/Private/ToolMenus.cpp
@@ -14,6 +14,8 @@
 #include "HAL/PlatformApplicationMisc.h" // For clipboard
 #include "Widgets/Layout/SScrollBox.h"
 
+#include UE_INLINE_GENERATED_CPP_BY_NAME(ToolMenus)
+
 #define LOCTEXT_NAMESPACE "ToolMenuSubsystem"
 
 DEFINE_LOG_CATEGORY(LogToolMenus);
@@ -2489,3 +2491,4 @@ void UToolMenus::RemoveAllCustomizations()
 }
 
 #undef LOCTEXT_NAMESPACE
+
diff --git a/Engine/Source/Developer/ToolMenus/Private/ToolMenusBlueprintLibrary.cpp b/Engine/Source/Developer/ToolMenus/Private/ToolMenusBlueprintLibrary.cpp
--- a/Engine/Source/Developer/ToolMenus/Private/ToolMenusBlueprintLibrary.cpp
+++ b/Engine/Source/Developer/ToolMenus/Private/ToolMenusBlueprintLibrary.cpp
@@ -3,6 +3,8 @@
 
 #include "ToolMenusBlueprintLibrary.h"
 
+#include UE_INLINE_GENERATED_CPP_BY_NAME(ToolMenusBlueprintLibrary)
+
 FScriptSlateIcon UToolMenuEntryExtensions::MakeScriptSlateIcon(const FName StyleSetName, const FName StyleName, const FName SmallStyleName)
 {
 	if (SmallStyleName == NAME_None)
@@ -120,3 +122,4 @@ void UToolMenuSectionExtensions::AddEntryObject(UPARAM(ref) FToolMenuSection& Se
 {
 	Section.AddEntryObject(InObject);
 }
+
diff --git a/Engine/Source/Developer/Virtualization/Private/VirtualizationFilterSettings.cpp b/Engine/Source/Developer/Virtualization/Private/VirtualizationFilterSettings.cpp
--- a/Engine/Source/Developer/Virtualization/Private/VirtualizationFilterSettings.cpp
+++ b/Engine/Source/Developer/Virtualization/Private/VirtualizationFilterSettings.cpp
@@ -2,8 +2,11 @@
 
 #include "VirtualizationFilterSettings.h"
 
+#include UE_INLINE_GENERATED_CPP_BY_NAME(VirtualizationFilterSettings)
+
 UVirtualizationFilterSettings::UVirtualizationFilterSettings(const FObjectInitializer& ObjectInitializer)
 	: Super(ObjectInitializer)
 {
 
 }
+
diff --git a/Engine/Source/Developer/VisualGraphUtils/Private/VisualGraphEdge.cpp b/Engine/Source/Developer/VisualGraphUtils/Private/VisualGraphEdge.cpp
--- a/Engine/Source/Developer/VisualGraphUtils/Private/VisualGraphEdge.cpp
+++ b/Engine/Source/Developer/VisualGraphUtils/Private/VisualGraphEdge.cpp
@@ -3,6 +3,8 @@
 #include "VisualGraphEdge.h"
 #include "VisualGraph.h"
 
+#include UE_INLINE_GENERATED_CPP_BY_NAME(VisualGraphEdge)
+
 ///////////////////////////////////////////////////////////////////////////////////
 /// FVisualGraphEdge
 ///////////////////////////////////////////////////////////////////////////////////
@@ -62,4 +64,4 @@ FString FVisualGraphEdge::DumpDot(const FVisualGraph* InGraph, int32 InIndendati
 
 	return FString::Printf(TEXT("%s%s %s %s%s;\n"),
 		*Indentation, *SourceNodeName, *ConnectorContent, *TargetNodeName, *AttributeContent);
-}
\ No newline at end of file
+}
diff --git a/Engine/Source/Developer/VisualGraphUtils/Private/VisualGraphElement.cpp b/Engine/Source/Developer/VisualGraphUtils/Private/VisualGraphElement.cpp
--- a/Engine/Source/Developer/VisualGraphUtils/Private/VisualGraphElement.cpp
+++ b/Engine/Source/Developer/VisualGraphUtils/Private/VisualGraphElement.cpp
@@ -3,6 +3,8 @@
 #include "VisualGraphElement.h"
 #include "VisualGraph.h"
 
+#include UE_INLINE_GENERATED_CPP_BY_NAME(VisualGraphElement)
+
 ///////////////////////////////////////////////////////////////////////////////////
 /// FVisualGraphElement
 ///////////////////////////////////////////////////////////////////////////////////
@@ -220,4 +222,4 @@ FString FVisualGraphElement::DumpDotStyle(const TOptional<EVisualGraphStyle>& In
 		return FString();
 	}
 	return StaticEnum<EVisualGraphStyle>()->GetDisplayNameTextByValue((int64)InStyle.GetValue()).ToString().ToLower();
-}
\ No newline at end of file
+}
diff --git a/Engine/Source/Developer/Windows/LiveCoding/Private/LiveCodingSettings.cpp b/Engine/Source/Developer/Windows/LiveCoding/Private/LiveCodingSettings.cpp
--- a/Engine/Source/Developer/Windows/LiveCoding/Private/LiveCodingSettings.cpp
+++ b/Engine/Source/Developer/Windows/LiveCoding/Private/LiveCodingSettings.cpp
@@ -3,6 +3,8 @@
 #include "LiveCodingSettings.h"
 #include "Misc/App.h"
 
+#include UE_INLINE_GENERATED_CPP_BY_NAME(LiveCodingSettings)
+
 ULiveCodingSettings::ULiveCodingSettings(const FObjectInitializer& ObjectInitializer) : Super(ObjectInitializer)
 {
 	FProperty* EngineModulesProperty = StaticClass()->FindPropertyByName("bPreloadEngineModules");
@@ -20,3 +22,4 @@ ULiveCodingSettings::ULiveCodingSettings(const FObjectInitializer& ObjectInitial
 	bPreloadProjectModules = true;
 	bPreloadProjectPluginModules = true;
 }
+
