Allows DumpGPU to be triggered on automated screenshot test

diff --git a/Engine/Source/Developer/FunctionalTesting/Private/ScreenshotFunctionalTestBase.cpp b/Engine/Source/Developer/FunctionalTesting/Private/ScreenshotFunctionalTestBase.cpp
--- a/Engine/Source/Developer/FunctionalTesting/Private/ScreenshotFunctionalTestBase.cpp
+++ b/Engine/Source/Developer/FunctionalTesting/Private/ScreenshotFunctionalTestBase.cpp
@@ -18,6 +18,7 @@
 #include "Tests/AutomationCommon.h"
 #include "Logging/LogMacros.h"
 #include "UObject/AutomationObjectVersion.h"
+#include "RenderGraphBuilder.h"
 
 #include UE_INLINE_GENERATED_CPP_BY_NAME(ScreenshotFunctionalTestBase)
 
@@ -25,6 +26,11 @@
 
 DEFINE_LOG_CATEGORY_STATIC(LogScreenshotFunctionalTest, Log, Log)
 
+static TAutoConsoleVariable<int32> GDumpGPUDumpOnScreenshotTest(
+	TEXT("r.DumpGPU.DumpOnScreenshotTest"), 0,
+	TEXT("Allows to filter the tree when using r.DumpGPU command, the pattern match is case sensitive."),
+	ECVF_Default);
+
 AScreenshotFunctionalTestBase::AScreenshotFunctionalTestBase(const FObjectInitializer& ObjectInitializer)
 	: AFunctionalTest(ObjectInitializer)
 	, ScreenshotOptions(EComparisonTolerance::Low)
@@ -117,6 +123,15 @@ void AScreenshotFunctionalTestBase::PrepareForScreenshot()
 			bNeedsViewportRestore = true;
 		}
 	}
+
+#if WITH_DUMPGPU
+	// Reset render target extent to reduce size of the dumpGPU.
+	if (GDumpGPUDumpOnScreenshotTest.GetValueOnGameThread() != 0)
+	{
+		FlushRenderingCommands();
+		UKismetSystemLibrary::ExecuteConsoleCommand(GEngine->GameViewport->GetWorld(), TEXT("r.ResetRenderTargetsExtent"), nullptr);
+	}
+#endif
 #endif
 }
 
@@ -168,6 +183,13 @@ void AScreenshotFunctionalTestBase::RequestScreenshot()
 
 	UGameViewportClient* GameViewportClient = GEngine->GameViewport;
 	GameViewportClient->OnScreenshotCaptured().AddUObject(this, &AScreenshotFunctionalTestBase::OnScreenShotCaptured);
+
+#if WITH_AUTOMATION_TESTS && WITH_DUMPGPU
+	if (GDumpGPUDumpOnScreenshotTest.GetValueOnGameThread() != 0)
+	{
+		FRDGBuilder::BeginResourceDump(TEXT(""));
+	}
+#endif
 }
 
 void AScreenshotFunctionalTestBase::OnComparisonComplete(const FAutomationScreenshotCompareResults& CompareResults)
