Use TestLabel instead of GetName() when logging through Functional Test actor

#rnx

diff --git a/Engine/Source/Developer/FunctionalTesting/Private/FunctionalTest.cpp b/Engine/Source/Developer/FunctionalTesting/Private/FunctionalTest.cpp
--- a/Engine/Source/Developer/FunctionalTesting/Private/FunctionalTest.cpp
+++ b/Engine/Source/Developer/FunctionalTesting/Private/FunctionalTest.cpp
@@ -467,7 +467,7 @@ void AFunctionalTest::LogMessage(const FString& Message)
 	UE_LOG(LogFunctionalTest, Log, TEXT("%s"), *Message);
 	UE_VLOG(this, LogFunctionalTest, Log
 		, TEXT("%s> %s")
-		, *GetName(), *Message);
+		, *TestLabel, *Message);
 }
 
 void AFunctionalTest::SetTimeLimit(float InTimeLimit, EFunctionalTestResult InResult)
@@ -476,7 +476,7 @@ void AFunctionalTest::SetTimeLimit(float InTimeLimit, EFunctionalTestResult InRe
 	{
 		UE_VLOG(this, LogFunctionalTest, Warning
 			, TEXT("%s> Trying to set TimeLimit to less than 0. Falling back to 0 (infinite).")
-			, *GetName());
+			, *TestLabel);
 
 		InTimeLimit = 0.f;
 	}
@@ -486,7 +486,7 @@ void AFunctionalTest::SetTimeLimit(float InTimeLimit, EFunctionalTestResult InRe
 	{
 		UE_VLOG(this, LogFunctionalTest, Warning
 			, TEXT("%s> Trying to set test Result to \'Invalid\'. Falling back to \'Failed\'")
-			, *GetName());
+			, *TestLabel);
 
 		InResult = EFunctionalTestResult::Failed;
 	}
@@ -1102,7 +1102,7 @@ void AFunctionalTest::LogStep(ELogVerbosity::Type Verbosity, const FString& Mess
 {
 	TStringBuilder<256> FullMessage;
 
-	FullMessage.Append(GetName());
+	FullMessage.Append(TestLabel);
 	FullMessage.Append(TEXT(": "));
 	FullMessage.Append(Message);
 
@@ -1123,7 +1123,7 @@ void AFunctionalTest::LogStep(ELogVerbosity::Type Verbosity, const FString& Mess
 	// Warn if we do not have a current functional test. Such a situation prevents Warning/Error results from being associated with an actual test
 	if (!CurrentFunctionalTest)
 	{
-		UE_LOG(LogFunctionalTest, Warning, TEXT("FunctionalTest '%s' ran test '%s' when no functional test was active. This result will not be tracked."), *GetName(), *Message);
+		UE_LOG(LogFunctionalTest, Warning, TEXT("FunctionalTest '%s' ran test '%s' when no functional test was active. This result will not be tracked."), *TestLabel, *Message);
 	}
 
 	/* 
