Check if FunctionalTest is in /Game/ before removing the string from the path
+ update test full names

#rnx

diff --git a/Engine/Source/Developer/FunctionalTesting/Private/FunctionalTestingModule.cpp b/Engine/Source/Developer/FunctionalTesting/Private/FunctionalTestingModule.cpp
--- a/Engine/Source/Developer/FunctionalTesting/Private/FunctionalTestingModule.cpp
+++ b/Engine/Source/Developer/FunctionalTesting/Private/FunctionalTestingModule.cpp
@@ -129,7 +129,11 @@ void FFunctionalTestingModule::GetMapTests(bool bEditorOnlyTests, TArray<FString
 				FString MapAssetPath = MapAsset.GetObjectPathString();
 				FString MapPackageName = MapAsset.PackageName.ToString();
 				if (!IsDeveloperDirectoryIncluded && MapPackageName.Find(TEXT("/Game/Developers")) == 0) continue;
-				FString PartialSuiteName = MapPackageName.RightChop(6).Replace(TEXT("/"), TEXT(".")); // Remove "/Game/" from the name and use dot syntax
+				FString PartialSuiteName = MapPackageName.RightChop(1).Replace(TEXT("/"), TEXT(".")); // use dot syntax
+				if (MapPackageName.StartsWith(TEXT("/Game/")))
+				{
+					PartialSuiteName = PartialSuiteName.RightChop(5); // Remove "/Game/" from the name
+				}
 
 				FString AllTestNames;
 				FAssetDataTagMapSharedView::FFindTagResult MapTestNames = MapAsset.TagsAndValues.FindTag(bEditorOnlyTests ? TEXT("TestNamesEditor") : TEXT("TestNames"));
