Fix crash when using Automation Wait node in BP when AutomationController is not loaded.

diff --git a/Engine/Source/Developer/FunctionalTesting/Private/AutomationBlueprintFunctionLibrary.cpp b/Engine/Source/Developer/FunctionalTesting/Private/AutomationBlueprintFunctionLibrary.cpp
--- a/Engine/Source/Developer/FunctionalTesting/Private/AutomationBlueprintFunctionLibrary.cpp
+++ b/Engine/Source/Developer/FunctionalTesting/Private/AutomationBlueprintFunctionLibrary.cpp
@@ -713,7 +713,11 @@ void UAutomationBlueprintFunctionLibrary::FinishLoadingBeforeScreenshot()
 	{
 		UMaterialInterface::SubmitRemainingJobsForWorld(CurrentWorld);
 		FAssetCompilingManager::Get().FinishAllCompilation();
-		FModuleManager::GetModuleChecked<IAutomationControllerModule>("AutomationController").GetAutomationController()->ResetAutomationTestTimeout(TEXT("shader compilation"));
+		IAutomationControllerModule* AutomationControllerModule = FModuleManager::GetModulePtr<IAutomationControllerModule>("AutomationController");
+		if (AutomationControllerModule != nullptr)
+		{
+			AutomationControllerModule->GetAutomationController()->ResetAutomationTestTimeout(TEXT("shader compilation"));
+		}
 	}
 
 	// Force all mip maps to load before taking the screenshot.
