FunctionalTestingModule: Cook determinism: Sort AssetRegistryTag value for TestNames.
#rnx

diff --git a/Engine/Source/Developer/FunctionalTesting/Private/FunctionalTestingModule.cpp b/Engine/Source/Developer/FunctionalTesting/Private/FunctionalTestingModule.cpp
--- a/Engine/Source/Developer/FunctionalTesting/Private/FunctionalTestingModule.cpp
+++ b/Engine/Source/Developer/FunctionalTesting/Private/FunctionalTestingModule.cpp
@@ -63,7 +63,8 @@ void FFunctionalTestingModule::ShutdownModule()
 void FFunctionalTestingModule::OnGetAssetTagsForWorld(const UWorld* World, TArray<UObject::FAssetRegistryTag>& OutTags)
 {
 #if WITH_EDITOR
-	FString TestNames, TestNamesEditor;
+	TArray<FString> TestNamesRuntime;
+	TArray<FString> TestNamesEditor;
 	for (TActorIterator<AFunctionalTest> ActorItr(const_cast<UWorld*>(World), AFunctionalTest::StaticClass(), EActorIteratorFlags::AllActors); ActorItr; ++ActorItr)
 	{
 		AFunctionalTest* FunctionalTest = *ActorItr;
@@ -73,26 +74,23 @@ void FFunctionalTestingModule::OnGetAssetTagsForWorld(const UWorld* World, TArra
 			// Only include enabled tests in the list of functional tests to run.
 			if (FunctionalTest->IsEnabled())
 			{
-				bool bIsEditorOnly = IsEditorOnlyObject(FunctionalTest);
-
-				// Check if this class is editor only
-				FString& NamesAppend = bIsEditorOnly ? TestNamesEditor : TestNames;
-
-				NamesAppend.Append(FunctionalTest->GetActorLabel() + TEXT("|") + FunctionalTest->GetName());
-				NamesAppend.Append(TEXT(";"));
+				TArray<FString>& TestNames = IsEditorOnlyObject(FunctionalTest) ? TestNamesEditor : TestNamesRuntime;
+				TestNames.Add(FString::Printf(TEXT("%s|%s;"), *FunctionalTest->GetActorLabel(), *FunctionalTest->GetName()));
 			}
 		}
 	}
 
-	if (!TestNames.IsEmpty())
-	{
-		OutTags.Add(UObject::FAssetRegistryTag("TestNames", TestNames, UObject::FAssetRegistryTag::TT_Hidden));
-	}
-
-	if (!TestNamesEditor.IsEmpty())
+	auto AddTestNames = [&OutTags](const TCHAR* TagName, TArray<FString>& TestNames)
 	{
-		OutTags.Add(UObject::FAssetRegistryTag("TestNamesEditor", TestNamesEditor, UObject::FAssetRegistryTag::TT_Hidden));
-	}
+		if (!TestNames.IsEmpty())
+		{
+			TestNames.Sort();
+			FString TestNamesStr = FString::Join(TestNames, TEXT(""));
+			OutTags.Add(UObject::FAssetRegistryTag(TagName, MoveTemp(TestNamesStr), UObject::FAssetRegistryTag::TT_Hidden));
+		}
+	};
+	AddTestNames(TEXT("TestNames"), TestNamesRuntime);
+	AddTestNames(TEXT("TestNamesEditor"), TestNamesEditor);
 #endif
 }
 
