AFunctionalTest: No support for DataLayers

diff --git a/Engine/Source/Developer/FunctionalTesting/Classes/FunctionalTest.h b/Engine/Source/Developer/FunctionalTesting/Classes/FunctionalTest.h
--- a/Engine/Source/Developer/FunctionalTesting/Classes/FunctionalTest.h
+++ b/Engine/Source/Developer/FunctionalTesting/Classes/FunctionalTest.h
@@ -665,7 +665,8 @@ public:
 	virtual void Tick(float DeltaSeconds) override;
 	virtual void EndPlay(const EEndPlayReason::Type EndPlayReason) override;
 #if WITH_EDITOR
-	virtual bool CanChangeIsSpatiallyLoadedFlag() const { return false; }
+	virtual bool CanChangeIsSpatiallyLoadedFlag() const override { return false; }
+	virtual bool ActorTypeSupportsDataLayer() const override { return false; }
 #endif
 	// AActor interface end
 
