Output Functional Tests metadata in test info during execution in the log
Use Automation.LogBPTestMetadata CVar to enable the behavior

diff --git a/Engine/Plugins/Tests/EditorTests/Source/EditorTests/Private/EditorUtilityTest.cpp b/Engine/Plugins/Tests/EditorTests/Source/EditorTests/Private/EditorUtilityTest.cpp
--- a/Engine/Plugins/Tests/EditorTests/Source/EditorTests/Private/EditorUtilityTest.cpp
+++ b/Engine/Plugins/Tests/EditorTests/Source/EditorTests/Private/EditorUtilityTest.cpp
@@ -75,6 +75,12 @@ void UEditorUtilityTest::Run()
 			CurrentTest->SetEditorUtility(this);
 		}
 
+		if (FAutomationTestFramework::NeedLogBPTestMetadata())
+		{
+			AddInfo(FString::Printf(TEXT("[Owner] %s"), *Owner));
+			AddInfo(FString::Printf(TEXT("[Description] %s"), *Description));
+		}
+
 		PrepareTest();
 		OnTestPrepare.Broadcast();
 	}
diff --git a/Engine/Source/Developer/FunctionalTesting/Classes/FunctionalTest.h b/Engine/Source/Developer/FunctionalTesting/Classes/FunctionalTest.h
--- a/Engine/Source/Developer/FunctionalTesting/Classes/FunctionalTest.h
+++ b/Engine/Source/Developer/FunctionalTesting/Classes/FunctionalTest.h
@@ -235,11 +235,11 @@ public:
 	FString TestLabel;
 
 	/**
-	 * The author is the group or person responsible for the test.  Generally you should use a group name
-	 * like 'Editor Team' or 'Rendering Team'.  When a test fails it may not be obvious who should investigate
+	 * The owner is the group or person responsible for the test. Generally you should use a group name
+	 * like 'Editor' or 'Rendering'. When a test fails it may not be obvious who should investigate
 	 * so this provides a associate responsible groups with tests.
 	 */
-	UPROPERTY(EditAnywhere, BlueprintReadOnly, Category = "Functional Testing", meta = (MultiLine = "true"))
+	UPROPERTY(EditAnywhere, BlueprintReadOnly, Category = "Functional Testing", meta = (MultiLine = "true", DisplayName = "Owner"))
 	FString Author;
 
 	/**
@@ -590,6 +590,9 @@ public:
 	UFUNCTION(BlueprintCallable, Category = "Reporting")
 	virtual void AddError(const FString& Message);
 
+	UFUNCTION(BlueprintCallable, Category = "Reporting")
+	virtual void AddInfo(const FString& Message);
+
 //protected:
 	/** TODO: break this out into a library */
 	void LogStep(ELogVerbosity::Type Verbosity, const FString& Message);
diff --git a/Engine/Source/Developer/FunctionalTesting/Private/FunctionalTest.cpp b/Engine/Source/Developer/FunctionalTesting/Private/FunctionalTest.cpp
--- a/Engine/Source/Developer/FunctionalTesting/Private/FunctionalTest.cpp
+++ b/Engine/Source/Developer/FunctionalTesting/Private/FunctionalTest.cpp
@@ -295,6 +295,12 @@ void AFunctionalTest::StartTest()
 	StartFrame = GFrameNumber;
 	StartTime = (float)GetWorld()->GetTimeSeconds();
 
+	if (FAutomationTestFramework::NeedLogBPTestMetadata())
+	{
+		AddInfo(FString::Printf(TEXT("[Owner] %s"), *Author));
+		AddInfo(FString::Printf(TEXT("[Description] %s"), *Description));
+	}
+
 	ReceiveStartTest();
 	OnTestStart.Broadcast();
 }
@@ -1098,6 +1104,11 @@ void AFunctionalTest::AddError(const FString& Message)
 	LogStep(ELogVerbosity::Error, Message);
 }
 
+void AFunctionalTest::AddInfo(const FString& Message)
+{
+	LogStep(ELogVerbosity::Log, Message);
+}
+
 void AFunctionalTest::LogStep(ELogVerbosity::Type Verbosity, const FString& Message)
 {
 	TStringBuilder<256> FullMessage;
diff --git a/Engine/Source/Runtime/Core/Private/Misc/AutomationTest.cpp b/Engine/Source/Runtime/Core/Private/Misc/AutomationTest.cpp
--- a/Engine/Source/Runtime/Core/Private/Misc/AutomationTest.cpp
+++ b/Engine/Source/Runtime/Core/Private/Misc/AutomationTest.cpp
@@ -35,6 +35,12 @@ namespace AutomationTest
 		TEXT("Automation.SkipStackWalk"),
 		bSkipStackWalk,
 		TEXT("Whether to skip stack walk while iterating for listing the tests"));
+
+	static bool bLogBPTestMetadata = false;
+	static FAutoConsoleVariableRef CVarAutomationLogBPTestMetadata(
+		TEXT("Automation.LogBPTestMetadata"),
+		bLogBPTestMetadata,
+		TEXT("Whether to output blueprint functional test metadata to the log when test is running"));
 };
 
 
@@ -270,6 +276,11 @@ bool FAutomationTestFramework::NeedSkipStackWalk()
 }
 
 
+bool FAutomationTestFramework::NeedLogBPTestMetadata()
+{
+	return AutomationTest::bLogBPTestMetadata;
+}
+
 bool FAutomationTestFramework::RegisterAutomationTest( const FString& InTestNameToRegister, FAutomationTestBase* InTestToRegister )
 {
 	const bool bAlreadyRegistered = AutomationTestClassNameToInstanceMap.Contains( InTestNameToRegister );
diff --git a/Engine/Source/Runtime/Core/Public/Misc/AutomationTest.h b/Engine/Source/Runtime/Core/Public/Misc/AutomationTest.h
--- a/Engine/Source/Runtime/Core/Public/Misc/AutomationTest.h
+++ b/Engine/Source/Runtime/Core/Public/Misc/AutomationTest.h
@@ -1024,10 +1024,15 @@ public:
 	}
 
 	/**
-	* Whether to skip stack walk while iterating for listing the tests
-	*/
+	 * Whether to skip stack walk while iterating for listing the tests
+	 */
 	static bool NeedSkipStackWalk();
 
+	/**
+	 * Whether to output blueprint functional test metadata to the log when test is running
+	 */
+	static bool NeedLogBPTestMetadata();
+
 	void NotifyScreenshotComparisonComplete(const FAutomationScreenshotCompareResults& CompareResults);
 	void NotifyTestDataRetrieved(bool bWasNew, const FString& JsonData);
 	void NotifyPerformanceDataRetrieved(bool bSuccess, const FString& ErrorMessage);
