Fix "Functional test is already running" error by unsetting FunctionalTesting manager from IFunctionalTestingModule when all tests are done

diff --git a/Engine/Source/Developer/FunctionalTesting/Private/FunctionalTestingManager.cpp b/Engine/Source/Developer/FunctionalTesting/Private/FunctionalTestingManager.cpp
--- a/Engine/Source/Developer/FunctionalTesting/Private/FunctionalTestingManager.cpp
+++ b/Engine/Source/Developer/FunctionalTesting/Private/FunctionalTestingManager.cpp
@@ -297,6 +297,7 @@ void UFunctionalTestingManager::AllTestsDone()
 	{
 		OnTestsComplete.Broadcast();
 		bFinished = true;
+		IFunctionalTestingModule::Get().SetManager(nullptr);
 		RemoveFromRoot();
 	}
 }
