Add function to reset GroundTruthData

#rnx

diff --git a/Engine/Source/Developer/FunctionalTesting/Private/GroundTruthData.cpp b/Engine/Source/Developer/FunctionalTesting/Private/GroundTruthData.cpp
--- a/Engine/Source/Developer/FunctionalTesting/Private/GroundTruthData.cpp
+++ b/Engine/Source/Developer/FunctionalTesting/Private/GroundTruthData.cpp
@@ -111,3 +111,23 @@ void UGroundTruthData::PostEditChangeProperty(FPropertyChangedEvent& PropertyCha
 }
 
 #endif
+
+void UGroundTruthData::ResetObject()
+{
+	#if WITH_EDITOR
+
+		if (ObjectData)
+		{
+			ObjectData->Rename(nullptr, GetTransientPackage());
+			ObjectData = nullptr;
+
+			MarkPackageDirty();
+		}
+
+	#else
+		FAssetData GroundTruthAssetData(this);
+		FString GroundTruthPackageName = GroundTruthAssetData.PackageName.ToString();
+
+		UE_LOG(GroundTruthLog, Error, TEXT("Can't reset ground truth data outside of the editor, '%s'."), *GroundTruthPackageName);
+	#endif
+}	
\ No newline at end of file
diff --git a/Engine/Source/Developer/FunctionalTesting/Public/GroundTruthData.h b/Engine/Source/Developer/FunctionalTesting/Public/GroundTruthData.h
--- a/Engine/Source/Developer/FunctionalTesting/Public/GroundTruthData.h
+++ b/Engine/Source/Developer/FunctionalTesting/Public/GroundTruthData.h
@@ -33,6 +33,9 @@ public:
 	UFUNCTION(BlueprintCallable, Category = "Automation")
 	bool CanModify() const;
 
+	UFUNCTION(BlueprintCallable, Category = "Automation")
+	void ResetObject();
+
 #if WITH_EDITOR
 	virtual void PostEditChangeProperty(struct FPropertyChangedEvent& PropertyChangedEvent) override;
 #endif
