Avoid a crash with Screenshot Functional Tests in PIE client net mode

diff --git a/Engine/Source/Developer/FunctionalTesting/Private/AutomationBlueprintFunctionLibrary.cpp b/Engine/Source/Developer/FunctionalTesting/Private/AutomationBlueprintFunctionLibrary.cpp
--- a/Engine/Source/Developer/FunctionalTesting/Private/AutomationBlueprintFunctionLibrary.cpp
+++ b/Engine/Source/Developer/FunctionalTesting/Private/AutomationBlueprintFunctionLibrary.cpp
@@ -315,7 +315,7 @@ void FAutomationTestScreenshotEnvSetup::Setup(UWorld* InWorld, FAutomationScreen
 	AutomationViewExtension = FSceneViewExtensions::NewExtension<FAutomationViewExtension>(InWorld, InOutOptions, InCurrentTimeToSimulate);
 
 	// TODO - I don't like needing to set this here.  Because the gameviewport uses a console variable, it wins.
-	if (UGameViewportClient* ViewportClient = GEngine->GameViewport)
+	if (UGameViewportClient* ViewportClient = AutomationCommon::GetAnyGameViewportClient())
 	{
 		static IConsoleVariable* ICVar = IConsoleManager::Get().FindConsoleVariable(FBufferVisualizationData::GetVisualizationTargetConsoleCommandName());
 		if (ICVar)
@@ -352,7 +352,7 @@ void FAutomationTestScreenshotEnvSetup::Restore()
 
 	AutomationViewExtension.Reset();
 
-	if (UGameViewportClient* ViewportClient = GEngine->GameViewport)
+	if (UGameViewportClient* ViewportClient = AutomationCommon::GetAnyGameViewportClient())
 	{
 		static IConsoleVariable* ICVar = IConsoleManager::Get().FindConsoleVariable(FBufferVisualizationData::GetVisualizationTargetConsoleCommandName());
 		if (ICVar)
@@ -380,9 +380,10 @@ public:
 	{
 		EnvSetup.Setup(InWorld, Options);
 
+		UGameViewportClient* GameViewportClient = AutomationCommon::GetAnyGameViewportClient();
 		if (!FPlatformProperties::HasFixedResolution())
 		{
-			FSceneViewport* GameViewport = GEngine->GameViewport ? GEngine->GameViewport->GetGameViewport() : nullptr;
+			FSceneViewport* GameViewport = GameViewportClient ? GameViewportClient->GetGameViewport() : nullptr;
 			if (GameViewport)
 			{
 #if WITH_EDITOR
@@ -390,7 +391,7 @@ public:
 				UEditorEngine* EditorEngine = Cast<UEditorEngine>(GEngine);
 
 				const bool bIsPIEViewport = GameViewport->IsPlayInEditorViewport();
-				const bool bIsNewViewport = InWorld && EditorEngine && EditorEngine->WorldIsPIEInNewViewport(InWorld);
+				const bool bIsNewViewport = GameViewportClient->GetWorld() && EditorEngine && EditorEngine->WorldIsPIEInNewViewport(GameViewportClient->GetWorld());
 
 				if (!bIsPIEViewport || bIsNewViewport)
 #endif
@@ -405,7 +406,7 @@ public:
 
 		FlushRenderingCommands();
 
-		GEngine->GameViewport->OnScreenshotCaptured().AddRaw(this, &FAutomationScreenshotTaker::GrabScreenShot);
+		GameViewportClient->OnScreenshotCaptured().AddRaw(this, &FAutomationScreenshotTaker::GrabScreenShot);
 		FWorldDelegates::LevelRemovedFromWorld.AddRaw(this, &FAutomationScreenshotTaker::WorldDestroyed);
 		FScreenshotRequest::OnScreenshotRequestProcessed().AddRaw(this, &FAutomationScreenshotTaker::OnScreenshotProcessed);
 	}
@@ -415,19 +416,20 @@ public:
 		FAutomationTestFramework::Get().OnScreenshotCompared.RemoveAll(this);
 		FScreenshotRequest::OnScreenshotRequestProcessed().RemoveAll(this);
 
-		if (GEngine->GameViewport)
+		UGameViewportClient* GameViewportClient = AutomationCommon::GetAnyGameViewportClient();
+		if (GameViewportClient)
 		{
 			// remove before we restore the viewport's size - a resize can trigger a redraw, which would trigger OnScreenshotCaptured() again (endless loop)
-			GEngine->GameViewport->OnScreenshotCaptured().RemoveAll(this);
+			GameViewportClient->OnScreenshotCaptured().RemoveAll(this);
 		}
 
 		FWorldDelegates::LevelRemovedFromWorld.RemoveAll(this);
 
 		if (!FPlatformProperties::HasFixedResolution() && bNeedsViewportSizeRestore)
 		{
-			if (GEngine->GameViewport)
+			if (GameViewportClient)
 			{
-				FSceneViewport* GameViewport = GEngine->GameViewport->GetGameViewport();
+				FSceneViewport* GameViewport = GameViewportClient->GetGameViewport();
 				GameViewport->SetViewportSize(ViewportRestoreSize.X, ViewportRestoreSize.Y);
 			}
 		}
@@ -822,8 +824,14 @@ void UAutomationBlueprintFunctionLibrary::TakeAutomationScreenshotAtCamera(UObje
 		return;
 	}
 
-	APlayerController* PlayerController = UGameplayStatics::GetPlayerController(WorldContextObject, 0);
+	UGameViewportClient* GameViewportClient = AutomationCommon::GetAnyGameViewportClient();
+	if (GameViewportClient == nullptr)
+	{
+		FMessageLog("PIE").Error(LOCTEXT("GameViewportRequired", "No game viewport found in World to TakeAutomationScreenshotAtCamera. Use a delay or change Net mode to Standalone."));
+		return;
+	}
 
+	APlayerController* PlayerController = UGameplayStatics::GetPlayerController(GameViewportClient->GetWorld(), 0);
 	if ( PlayerController == nullptr )
 	{
 		FMessageLog("PIE").Error(LOCTEXT("PlayerRequired", "A player controller is required to TakeAutomationScreenshotAtCamera"));
@@ -832,6 +840,7 @@ void UAutomationBlueprintFunctionLibrary::TakeAutomationScreenshotAtCamera(UObje
 
 	// Move the player, then queue up a screenshot.
 	// We need to delay before the screenshot so that the motion blur has time to stop.
+	PlayerController->bAutoManageActiveCameraTarget = false;
 	PlayerController->SetViewTarget(Camera, FViewTargetTransitionParams());
 	FString ScreenshotName = Camera->GetName();
 
@@ -895,7 +904,10 @@ bool UAutomationBlueprintFunctionLibrary::TakeAutomationScreenshotOfUI_Immediate
 					Data.MaximumLocalError = Options.MaximumLocalError;
 					Data.MaximumGlobalError = Options.MaximumGlobalError;
 
-					GEngine->GameViewport->OnScreenshotCaptured().Broadcast(OutSize.X, OutSize.Y, OutColorData);
+					if (UGameViewportClient* GameViewportClient = AutomationCommon::GetAnyGameViewportClient())
+					{
+						GameViewportClient->OnScreenshotCaptured().Broadcast(OutSize.X, OutSize.Y, OutColorData);
+					}
 #endif
 
 					return true; //-V773
diff --git a/Engine/Source/Developer/FunctionalTesting/Private/FunctionalTestingManager.cpp b/Engine/Source/Developer/FunctionalTesting/Private/FunctionalTestingManager.cpp
--- a/Engine/Source/Developer/FunctionalTesting/Private/FunctionalTestingManager.cpp
+++ b/Engine/Source/Developer/FunctionalTesting/Private/FunctionalTestingManager.cpp
@@ -9,6 +9,7 @@
 #include "EngineGlobals.h"
 #include "Engine/Engine.h"
 #include "Misc/RuntimeErrors.h"
+#include "GameFramework/PlayerController.h"
 
 #include UE_INLINE_GENERATED_CPP_BY_NAME(FunctionalTestingManager)
 
@@ -166,7 +167,8 @@ void UFunctionalTestingManager::TriggerFirstValidTest()
 
 	const bool bIsWorldInitialized =
 		World->AreActorsInitialized() &&
-		(!World->GetWorldSettings()->IsNavigationSystemEnabled() || !UNavigationSystemV1::IsNavigationBeingBuilt(World));
+		(!World->GetWorldSettings()->IsNavigationSystemEnabled() || !UNavigationSystemV1::IsNavigationBeingBuilt(World)) &&
+		(World->GetNumPlayerControllers() != 0) && World->GetFirstPlayerController()->GetPawnOrSpectator() != nullptr;
 
 	if (bInitialDelayApplied == true && bIsWorldInitialized)
 	{
diff --git a/Engine/Source/Developer/FunctionalTesting/Private/ScreenshotFunctionalTest.cpp b/Engine/Source/Developer/FunctionalTesting/Private/ScreenshotFunctionalTest.cpp
--- a/Engine/Source/Developer/FunctionalTesting/Private/ScreenshotFunctionalTest.cpp
+++ b/Engine/Source/Developer/FunctionalTesting/Private/ScreenshotFunctionalTest.cpp
@@ -15,6 +15,7 @@
 #include "Slate/SceneViewport.h"
 #include "UObject/AutomationObjectVersion.h"
 #include "DataDrivenShaderPlatformInfo.h"
+#include "Tests/AutomationCommon.h"
 
 #include UE_INLINE_GENERATED_CPP_BY_NAME(ScreenshotFunctionalTest)
 
@@ -74,7 +75,9 @@ void AScreenshotFunctionalTest::RequestScreenshot()
 		FIntPoint ScreenshotViewportSize = UAutomationBlueprintFunctionLibrary::GetAutomationScreenshotSize(ScreenshotOptions);
 		if (Config.SetResolution(ScreenshotViewportSize.X, ScreenshotViewportSize.Y, 1.0f))
 		{
-			GEngine->GameViewport->GetGameViewport()->TakeHighResScreenShot();
+			UGameViewportClient* GameViewportClient = AutomationCommon::GetAnyGameViewportClient();
+			check(GameViewportClient);
+			GameViewportClient->GetGameViewport()->TakeHighResScreenShot();
 		}
 	}
 	else
diff --git a/Engine/Source/Developer/FunctionalTesting/Private/ScreenshotFunctionalTestBase.cpp b/Engine/Source/Developer/FunctionalTesting/Private/ScreenshotFunctionalTestBase.cpp
--- a/Engine/Source/Developer/FunctionalTesting/Private/ScreenshotFunctionalTestBase.cpp
+++ b/Engine/Source/Developer/FunctionalTesting/Private/ScreenshotFunctionalTestBase.cpp
@@ -49,9 +49,14 @@ void AScreenshotFunctionalTestBase::PrepareTest()
 {
 	Super::PrepareTest();
 
-	APlayerController* PlayerController = UGameplayStatics::GetPlayerController(GetWorld(), 0);
+	UGameViewportClient* GameViewportClient = AutomationCommon::GetAnyGameViewportClient();
+	check(GameViewportClient);
+
+	APlayerController* PlayerController = UGameplayStatics::GetPlayerController(GameViewportClient->GetWorld(), 0);
 	if (PlayerController)
 	{
+		// Make sure the camera target is not auto managed
+		PlayerController->bAutoManageActiveCameraTarget = false;
 		PlayerController->SetViewTarget(this, FViewTargetTransitionParams());
 	}
 
@@ -88,17 +93,18 @@ void AScreenshotFunctionalTestBase::OnScreenshotTakenAndCompared()
 
 void AScreenshotFunctionalTestBase::PrepareForScreenshot()
 {
-	check(GEngine->GameViewport && GEngine->GameViewport->GetGameViewport());
+	UGameViewportClient* GameViewportClient = AutomationCommon::GetAnyGameViewportClient();
+	check(GameViewportClient);
 	check(IsInGameThread());
 	check(!bNeedsViewSettingsRestore && !bNeedsViewportRestore);
 
 #if WITH_AUTOMATION_TESTS
 	bool bApplyScreenshotSettings = true;
-	FSceneViewport* GameViewport = GEngine->GameViewport->GetGameViewport();
+	FSceneViewport* GameViewport = GameViewportClient->GetGameViewport();
 
 #if WITH_EDITOR_AUTOMATION_TESTS
 	// In the editor we can only attempt to resize a standalone viewport
-	UWorld* World = GetWorld();
+	UWorld* World = GameViewportClient->GetWorld();
 	UEditorEngine* EditorEngine = Cast<UEditorEngine>(GEngine);	
 
 	const bool bIsPIEViewport = GameViewport->IsPlayInEditorViewport();
@@ -109,7 +115,7 @@ void AScreenshotFunctionalTestBase::PrepareForScreenshot()
 
 	if (bApplyScreenshotSettings)
 	{
-		ScreenshotEnvSetup->Setup(GetWorld(), ScreenshotOptions);
+		ScreenshotEnvSetup->Setup(GameViewportClient->GetWorld(), ScreenshotOptions);
 		FlushRenderingCommands();
 		bNeedsViewSettingsRestore = true;
 
@@ -129,7 +135,7 @@ void AScreenshotFunctionalTestBase::PrepareForScreenshot()
 	if (GDumpGPUDumpOnScreenshotTest.GetValueOnGameThread() != 0)
 	{
 		FlushRenderingCommands();
-		UKismetSystemLibrary::ExecuteConsoleCommand(GEngine->GameViewport->GetWorld(), TEXT("r.ResetRenderTargetsExtent"), nullptr);
+		UKismetSystemLibrary::ExecuteConsoleCommand(GameViewportClient->GetWorld(), TEXT("r.ResetRenderTargetsExtent"), nullptr);
 	}
 #endif
 #endif
@@ -137,9 +143,10 @@ void AScreenshotFunctionalTestBase::PrepareForScreenshot()
 
 void AScreenshotFunctionalTestBase::OnScreenShotCaptured(int32 InSizeX, int32 InSizeY, const TArray<FColor>& InImageData)
 {
-	check(GEngine->GameViewport);
+	UGameViewportClient* GameViewportClient = AutomationCommon::GetAnyGameViewportClient();
+	check(GameViewportClient);
 
-	GEngine->GameViewport->OnScreenshotCaptured().RemoveAll(this);
+	GameViewportClient->OnScreenshotCaptured().RemoveAll(this);
 
 #if WITH_AUTOMATION_TESTS
 	TArray<uint8> CapturedFrameTrace = AutomationCommon::CaptureFrameTrace(GetWorld()->GetName(), TestLabel);
@@ -176,12 +183,12 @@ void AScreenshotFunctionalTestBase::OnScreenShotCaptured(int32 InSizeX, int32 In
 void AScreenshotFunctionalTestBase::RequestScreenshot()
 {
 	check(IsInGameThread());
-	check(GEngine->GameViewport);
+	UGameViewportClient* GameViewportClient = AutomationCommon::GetAnyGameViewportClient();
+	check(GameViewportClient);
 
 	// Make sure any screenshot request has been processed
 	FlushRenderingCommands();
 
-	UGameViewportClient* GameViewportClient = GEngine->GameViewport;
 	GameViewportClient->OnScreenshotCaptured().AddUObject(this, &AScreenshotFunctionalTestBase::OnScreenShotCaptured);
 
 #if WITH_AUTOMATION_TESTS && WITH_DUMPGPU
@@ -206,7 +213,8 @@ void AScreenshotFunctionalTestBase::OnComparisonComplete(const FAutomationScreen
 
 void AScreenshotFunctionalTestBase::RestoreViewSettings()
 {
-	check(GEngine->GameViewport && GEngine->GameViewport->GetGameViewport());
+	UGameViewportClient* GameViewportClient = AutomationCommon::GetAnyGameViewportClient();
+	check(GameViewportClient && GameViewportClient->GetGameViewport());
 	check(IsInGameThread());
 	
 #if WITH_AUTOMATION_TESTS
@@ -217,7 +225,7 @@ void AScreenshotFunctionalTestBase::RestoreViewSettings()
 
 	if (!FPlatformProperties::HasFixedResolution() && bNeedsViewportRestore)
 	{	
-		FSceneViewport* GameViewport = GEngine->GameViewport->GetGameViewport();
+		FSceneViewport* GameViewport = GameViewportClient->GetGameViewport();
 		GameViewport->SetViewportSize(ViewportRestoreSize.X, ViewportRestoreSize.Y);
 	}
 #endif
diff --git a/Engine/Source/Runtime/Engine/Private/Tests/AutomationCommon.cpp b/Engine/Source/Runtime/Engine/Private/Tests/AutomationCommon.cpp
--- a/Engine/Source/Runtime/Engine/Private/Tests/AutomationCommon.cpp
+++ b/Engine/Source/Runtime/Engine/Private/Tests/AutomationCommon.cpp
@@ -349,6 +349,25 @@ namespace AutomationCommon
 
 		return TestWorld;
 	}
+
+	UGameViewportClient* GetAnyGameViewportClient()
+	{
+		if (GEngine->GameViewport)
+		{
+			return GEngine->GameViewport;
+		}
+		// Then Game viewport is attached to another world context than the main Engine one. (ie: PIE Net mode set to As Client)
+		const TIndirectArray<FWorldContext>& WorldContexts = GEngine->GetWorldContexts();
+		for (const FWorldContext& Context : WorldContexts)
+		{
+			if ((Context.WorldType == EWorldType::PIE) && Context.World() != nullptr && Context.GameViewport != nullptr)
+			{
+				return Context.GameViewport;
+			}
+		}
+
+		return nullptr;
+	}
 }
 
 bool AutomationOpenMap(const FString& MapName, bool bForceReload)
diff --git a/Engine/Source/Runtime/Engine/Public/Tests/AutomationCommon.h b/Engine/Source/Runtime/Engine/Public/Tests/AutomationCommon.h
--- a/Engine/Source/Runtime/Engine/Public/Tests/AutomationCommon.h
+++ b/Engine/Source/Runtime/Engine/Public/Tests/AutomationCommon.h
@@ -4,6 +4,7 @@
 
 #include "CoreMinimal.h"
 #include "Misc/AutomationTest.h"
+#include "Engine/GameViewportClient.h"
 
 class AActor;
 class SWindow;
@@ -61,6 +62,8 @@ namespace AutomationCommon
 	ENGINE_API SWidget* FindWidgetByTag(const FName Tag);
 
 	ENGINE_API UWorld* GetAnyGameWorld();
+
+	ENGINE_API UGameViewportClient* GetAnyGameViewportClient();
 #endif
 }
 
