non-unity error fix in Function UI test.

diff --git a/Engine/Source/Developer/FunctionalTesting/Private/FunctionalUIScreenshotTest.cpp b/Engine/Source/Developer/FunctionalTesting/Private/FunctionalUIScreenshotTest.cpp
--- a/Engine/Source/Developer/FunctionalTesting/Private/FunctionalUIScreenshotTest.cpp
+++ b/Engine/Source/Developer/FunctionalTesting/Private/FunctionalUIScreenshotTest.cpp
@@ -6,6 +6,7 @@
 #include "AutomationBlueprintFunctionLibrary.h"
 #include "Camera/CameraComponent.h"
 #include "Camera/PlayerCameraManager.h"
+#include "Framework/Application/SlateApplication.h"
 #include "Kismet/GameplayStatics.h"
 #include "Engine/Engine.h"
 #include "EngineGlobals.h"
