Add CVar setting node in AFunctionalTest class

diff --git a/Engine/Source/Developer/FunctionalTesting/Classes/FunctionalTest.h b/Engine/Source/Developer/FunctionalTesting/Classes/FunctionalTest.h
--- a/Engine/Source/Developer/FunctionalTesting/Classes/FunctionalTest.h
+++ b/Engine/Source/Developer/FunctionalTesting/Classes/FunctionalTest.h
@@ -220,6 +220,43 @@ enum class EFunctionalTestLogHandling : uint8
 	OutputIgnored
 };
 
+
+class FConsoleVariableBPSetter
+{
+	friend class FAutomationFunctionalTestEnvSetup;
+
+public:
+	FConsoleVariableBPSetter(FString InConsoleVariableName);
+
+	void Set(const FString& Value);
+	FString Get();
+	void Restore();
+
+private:
+	bool bModified;
+	FString ConsoleVariableName;
+
+	FString OriginalValue;
+};
+
+class FAutomationFunctionalTestEnvSetup
+{
+public:
+	FAutomationFunctionalTestEnvSetup() = default;
+	~FAutomationFunctionalTestEnvSetup();
+
+	void SetVariable(const FString& VariableName, const FString& Value);
+
+	FString GetVariable(const FString& VariableName);
+
+	/** Restore the old settings. */
+	void Restore();
+
+private:
+	TArray<FConsoleVariableBPSetter> Variables;
+};
+
+
 DECLARE_DYNAMIC_MULTICAST_DELEGATE(FFunctionalTestEventSignature);
 DECLARE_DELEGATE_OneParam(FFunctionalTestDoneSignature, class AFunctionalTest*);
 
@@ -637,6 +674,22 @@ public:
 	UFUNCTION(BlueprintCallable, Category = "Functional Testing")
 	FName GetCurrentRerunReason() const;
 
+	/** Sets the CVar from the given input. Variable gets reset after the test. */
+	UFUNCTION(BlueprintCallable, Category = "Functional Testing")
+	void SetConsoleVariable(const FString& Name, const FString& InValue);
+
+	/** Sets the CVar from the given input. Variable gets reset after the test. */
+	UFUNCTION(BlueprintCallable, Category = "Functional Testing")
+	void SetConsoleVariableFromInteger(const FString& Name, const int32 InValue);
+
+	/** Sets the CVar from the given input. Variable gets reset after the test. */
+	UFUNCTION(BlueprintCallable, Category = "Functional Testing")
+	void SetConsoleVariableFromFloat(const FString& Name, const float InValue);
+
+	/** Sets the CVar from the given input. Variable gets reset after the test. */
+	UFUNCTION(BlueprintCallable, Category = "Functional Testing")
+	void SetConsoleVariableFromBoolean(const FString& Name, const bool InValue);
+
 	UFUNCTION(BlueprintImplementableEvent, Category = "Functional Testing")
 	FString OnAdditionalTestFinishedMessageRequest(EFunctionalTestResult TestResult) const;
 	
@@ -744,6 +797,7 @@ public:
 
 private:
 	bool bIsReady;
+	FAutomationFunctionalTestEnvSetup EnvSetup;
 
 public:
 	/** Returns SpriteComponent subobject **/
diff --git a/Engine/Source/Developer/FunctionalTesting/Private/FunctionalTest.cpp b/Engine/Source/Developer/FunctionalTesting/Private/FunctionalTest.cpp
--- a/Engine/Source/Developer/FunctionalTesting/Private/FunctionalTest.cpp
+++ b/Engine/Source/Developer/FunctionalTesting/Private/FunctionalTest.cpp
@@ -436,6 +436,8 @@ void AFunctionalTest::FinishTest(EFunctionalTestResult TestResult, const FString
 	//}
 
 	TestFinishedObserver.ExecuteIfBound(this);
+
+	EnvSetup.Restore();
 }
 
 void AFunctionalTest::EndPlay(const EEndPlayReason::Type EndPlayReason)
@@ -527,6 +529,26 @@ FName AFunctionalTest::GetCurrentRerunReason()const
 	return CurrentRerunCause;
 }
 
+void AFunctionalTest::SetConsoleVariable(const FString& Name, const FString& InValue)
+{
+	EnvSetup.SetVariable(Name, InValue);
+}
+
+void AFunctionalTest::SetConsoleVariableFromInteger(const FString& Name, const int32 InValue)
+{
+	EnvSetup.SetVariable(Name, FString::FromInt(InValue));
+}
+
+void AFunctionalTest::SetConsoleVariableFromFloat(const FString& Name, const float InValue)
+{
+	EnvSetup.SetVariable(Name, FString::SanitizeFloat(InValue));
+}
+
+void AFunctionalTest::SetConsoleVariableFromBoolean(const FString& Name, const bool InValue)
+{
+	EnvSetup.SetVariable(Name, FString::FromInt(InValue));
+}
+
 void AFunctionalTest::RegisterAutoDestroyActor(AActor* ActorToAutoDestroy)
 {
 	AutoDestroyActors.AddUnique(ActorToAutoDestroy);
@@ -1665,3 +1687,93 @@ void UAutomationPerformaceHelper::EndStatsFile()
 	}
 }
 
+FConsoleVariableBPSetter::FConsoleVariableBPSetter(FString InConsoleVariableName)
+	: bModified(false)
+	, ConsoleVariableName(InConsoleVariableName)
+{
+}
+
+void FConsoleVariableBPSetter::Set(const FString& Value)
+{
+	IConsoleVariable* ConsoleVariable = IConsoleManager::Get().FindConsoleVariable(*ConsoleVariableName);
+	if (ensure(ConsoleVariable))
+	{
+		if (bModified == false)
+		{
+			bModified = true;
+			OriginalValue = ConsoleVariable->GetString();
+		}
+
+		ConsoleVariable->AsVariable()->SetWithCurrentPriority(Value.GetCharArray().GetData());
+	}
+}
+
+FString FConsoleVariableBPSetter::Get()
+{
+	IConsoleVariable* ConsoleVariable = IConsoleManager::Get().FindConsoleVariable(*ConsoleVariableName);
+
+	if (ensure(ConsoleVariable))
+	{
+		return ConsoleVariable->GetString();
+	}
+
+	return FString{};
+}
+
+void FConsoleVariableBPSetter::Restore()
+{
+	if (bModified)
+	{
+		IConsoleVariable* ConsoleVariable = IConsoleManager::Get().FindConsoleVariable(*ConsoleVariableName);
+		if (ensure(ConsoleVariable))
+		{
+			ConsoleVariable->AsVariable()->SetWithCurrentPriority(OriginalValue.GetCharArray().GetData());
+		}
+
+		bModified = false;
+	}
+}
+
+FAutomationFunctionalTestEnvSetup::~FAutomationFunctionalTestEnvSetup()
+{
+	Restore();
+}
+
+void FAutomationFunctionalTestEnvSetup::SetVariable(const FString& VariableName, const FString& Value)
+{
+	check(IsInGameThread());
+
+	FConsoleVariableBPSetter Variable(VariableName);
+	Variable.Set(Value);
+	Variables.Add(MoveTemp(Variable));
+}
+
+FString FAutomationFunctionalTestEnvSetup::GetVariable(const FString& VariableName)
+{
+	check(IsInGameThread());
+
+	for (auto& Variable : Variables)
+	{
+		if (Variable.ConsoleVariableName == VariableName)
+		{
+			return Variable.Get();
+		}
+	}
+
+	return FString{};
+}
+
+void FAutomationFunctionalTestEnvSetup::Restore()
+{
+	if (!Variables.IsEmpty())
+	{
+		check(IsInGameThread());
+
+		for (auto& Variable : Variables)
+		{
+			Variable.Restore();
+		}
+
+		Variables.Empty();
+	}
+}
