Addressing casing issues with Automation test results by using the casing from the map's path instead of the map name itself since that is more likely to be a different case on clients.

diff --git a/Engine/Plugins/MovieScene/MovieRenderPipeline/Source/MovieRenderPipelineEditor/Private/MoviePipelineFunctionalTestBase.cpp b/Engine/Plugins/MovieScene/MovieRenderPipeline/Source/MovieRenderPipelineEditor/Private/MoviePipelineFunctionalTestBase.cpp
--- a/Engine/Plugins/MovieScene/MovieRenderPipeline/Source/MovieRenderPipelineEditor/Private/MoviePipelineFunctionalTestBase.cpp
+++ b/Engine/Plugins/MovieScene/MovieRenderPipeline/Source/MovieRenderPipelineEditor/Private/MoviePipelineFunctionalTestBase.cpp
@@ -153,7 +153,7 @@ void AMoviePipelineFunctionalTestBase::CompareRenderOutputToGroundTruth(FMoviePi
 	check(OutputSetting);
 
 	// Build screenshot data for this test. This contains a lot of metadata about RHI, Platform, etc that we'll use to generate the output folder name.
-	FAutomationScreenshotData Data = UAutomationBlueprintFunctionLibrary::BuildScreenshotData(GetWorld()->GetName(), TestLabel, OutputSetting->OutputResolution.X, OutputSetting->OutputResolution.Y);
+	FAutomationScreenshotData Data = UAutomationBlueprintFunctionLibrary::BuildScreenshotData(GetWorld(), TestLabel, OutputSetting->OutputResolution.X, OutputSetting->OutputResolution.Y);
 
 	// Convert the Screenshot Data into Metadata
 	FAutomationScreenshotMetadata MetaData(Data);
diff --git a/Engine/Source/Developer/FunctionalTesting/Private/AutomationBlueprintFunctionLibrary.cpp b/Engine/Source/Developer/FunctionalTesting/Private/AutomationBlueprintFunctionLibrary.cpp
--- a/Engine/Source/Developer/FunctionalTesting/Private/AutomationBlueprintFunctionLibrary.cpp
+++ b/Engine/Source/Developer/FunctionalTesting/Private/AutomationBlueprintFunctionLibrary.cpp
@@ -447,7 +447,7 @@ public:
 
 		if (World.IsValid())
 		{
-			FAutomationScreenshotData Data = UAutomationBlueprintFunctionLibrary::BuildScreenshotData(World->GetName(), ScreenShotName, InSizeX, InSizeY);
+			FAutomationScreenshotData Data = UAutomationBlueprintFunctionLibrary::BuildScreenshotData(World.Get(), ScreenShotName, InSizeX, InSizeY);
 
 			// Copy the relevant data into the metadata for the screenshot.
 			Data.bHasComparisonRules = true;
@@ -775,6 +775,11 @@ FAutomationScreenshotData UAutomationBlueprintFunctionLibrary::BuildScreenshotDa
 #endif
 }
 
+FAutomationScreenshotData UAutomationBlueprintFunctionLibrary::BuildScreenshotData(UWorld* InWorld, const FString& ScreenShotName, int32 Width, int32 Height)
+{
+	return BuildScreenshotData(AutomationCommon::GetWorldContext(InWorld), ScreenShotName, Width, Height);
+}
+
 bool UAutomationBlueprintFunctionLibrary::TakeAutomationScreenshotInternal(UObject* WorldContextObject, const FString& ScreenShotName, const FString& Notes, FAutomationScreenshotOptions Options)
 {
 	UAutomationBlueprintFunctionLibrary::FinishLoadingBeforeScreenshot();
diff --git a/Engine/Source/Developer/FunctionalTesting/Private/ScreenshotFunctionalTestBase.cpp b/Engine/Source/Developer/FunctionalTesting/Private/ScreenshotFunctionalTestBase.cpp
--- a/Engine/Source/Developer/FunctionalTesting/Private/ScreenshotFunctionalTestBase.cpp
+++ b/Engine/Source/Developer/FunctionalTesting/Private/ScreenshotFunctionalTestBase.cpp
@@ -149,9 +149,11 @@ void AScreenshotFunctionalTestBase::OnScreenShotCaptured(int32 InSizeX, int32 In
 	GameViewportClient->OnScreenshotCaptured().RemoveAll(this);
 
 #if WITH_AUTOMATION_TESTS
-	TArray<uint8> CapturedFrameTrace = AutomationCommon::CaptureFrameTrace(GetWorld()->GetName(), TestLabel);
+	const FString Context = AutomationCommon::GetWorldContext(GetWorld());
 
-	FAutomationScreenshotData Data = UAutomationBlueprintFunctionLibrary::BuildScreenshotData(GetWorld()->GetName(), TestLabel, InSizeX, InSizeY);
+	TArray<uint8> CapturedFrameTrace = AutomationCommon::CaptureFrameTrace(Context, TestLabel);
+
+	FAutomationScreenshotData Data = UAutomationBlueprintFunctionLibrary::BuildScreenshotData(Context, TestLabel, InSizeX, InSizeY);
 
 	// Copy the relevant data into the metadata for the screenshot.
 	Data.bHasComparisonRules = true;
diff --git a/Engine/Source/Developer/FunctionalTesting/Public/AutomationBlueprintFunctionLibrary.h b/Engine/Source/Developer/FunctionalTesting/Public/AutomationBlueprintFunctionLibrary.h
--- a/Engine/Source/Developer/FunctionalTesting/Public/AutomationBlueprintFunctionLibrary.h
+++ b/Engine/Source/Developer/FunctionalTesting/Public/AutomationBlueprintFunctionLibrary.h
@@ -80,6 +80,7 @@ public:
 	static bool TakeAutomationScreenshotInternal(UObject* WorldContextObject, const FString& ScreenShotName, const FString& Notes, FAutomationScreenshotOptions Options);
 
 	static FAutomationScreenshotData BuildScreenshotData(const FString& MapOrContext, const FString& ScreenShotName, int32 Width, int32 Height);
+	static FAutomationScreenshotData BuildScreenshotData(UWorld* InWorld, const FString& ScreenShotName, int32 Width, int32 Height);
 
 	static FIntPoint GetAutomationScreenshotSize(const FAutomationScreenshotOptions& Options);
 
diff --git a/Engine/Source/Runtime/Engine/Private/Tests/AutomationCommon.cpp b/Engine/Source/Runtime/Engine/Private/Tests/AutomationCommon.cpp
--- a/Engine/Source/Runtime/Engine/Private/Tests/AutomationCommon.cpp
+++ b/Engine/Source/Runtime/Engine/Private/Tests/AutomationCommon.cpp
@@ -367,6 +367,20 @@ namespace AutomationCommon
 
 		return nullptr;
 	}
+	
+	FString GetWorldContext(UWorld* InWorld)
+	{
+		FString Result = InWorld->GetName();
+
+		// Rely on the package loading path to avoid FName case inconsistencies with individual level names.
+		if (UPackage* Package = InWorld->GetPackage())
+		{
+			Result = FPaths::GetBaseFilename(Package->GetLoadedPath().GetPackageName());
+			checkf(FCString::Stricmp(*InWorld->GetName(), *Result) == 0, TEXT("Unexpected inconsistency between world and package path: expected \"%s\" got \"%s\""), *InWorld->GetName(), *Result);
+		}
+
+		return Result;
+	}
 }
 
 #if WITH_AUTOMATION_TESTS
diff --git a/Engine/Source/Runtime/Engine/Public/Tests/AutomationCommon.h b/Engine/Source/Runtime/Engine/Public/Tests/AutomationCommon.h
--- a/Engine/Source/Runtime/Engine/Public/Tests/AutomationCommon.h
+++ b/Engine/Source/Runtime/Engine/Public/Tests/AutomationCommon.h
@@ -65,6 +65,9 @@ namespace AutomationCommon
 
 #endif
 	ENGINE_API UGameViewportClient* GetAnyGameViewportClient();
+
+	/* Get the adjusted World name to use for screenshot paths */
+	ENGINE_API FString GetWorldContext(UWorld* InWorld);
 }
 
 #if WITH_AUTOMATION_TESTS
