Avoid Take High Res Screenshot scripting command to unset viewport user settings

#rnx

diff --git a/Engine/Source/Developer/FunctionalTesting/Private/AutomationBlueprintFunctionLibrary.cpp b/Engine/Source/Developer/FunctionalTesting/Private/AutomationBlueprintFunctionLibrary.cpp
--- a/Engine/Source/Developer/FunctionalTesting/Private/AutomationBlueprintFunctionLibrary.cpp
+++ b/Engine/Source/Developer/FunctionalTesting/Private/AutomationBlueprintFunctionLibrary.cpp
@@ -615,7 +615,9 @@ private:
 class FScreenshotTakenState : public FAutomationTaskStatusBase
 {
 public:
-	FScreenshotTakenState()
+	FScreenshotTakenState(bool InNeedGameViewToggle = true, bool InNeedCameraChange = true)
+		: NeedGameViewToggle(InNeedGameViewToggle)
+		, NeedCameraChange(InNeedCameraChange)
 	{
 		if (GIsAutomationTesting)
 		{
@@ -664,20 +666,38 @@ public:
 		{
 			FLevelEditorModule& LevelEditor = FModuleManager::GetModuleChecked<FLevelEditorModule>("LevelEditor");
 			SLevelViewport* LevelViewport = LevelEditor.GetFirstActiveLevelViewport().Get();
-			if (LevelViewport->IsInGameView() && LevelViewport->CanToggleGameView())
+			if (LevelViewport)
 			{
-				LevelViewport->ToggleGameView();
-			}
-			FLevelEditorViewportClient& LevelViewportClient = LevelViewport->GetLevelViewportClient();
-			if (LevelViewportClient.IsAnyActorLocked())
-			{
-				LevelViewportClient.SetActorLock(nullptr);
-				LevelViewportClient.bDisableInput = false;
-				LevelViewportClient.bEnableFading = true;
+				if (NeedGameViewToggle)
+				{
+					if (LevelViewport->IsInGameView() && LevelViewport->CanToggleGameView())
+					{
+						LevelViewport->ToggleGameView();
+					}
+					else
+					{
+						UE_LOG(AutomationFunctionLibrary, Verbose, TEXT("Expected to be able to toggle off the Game View mode after the screenshot was taken, but the Viewport was already no longer in that mode or it is not a Perspective."));
+					}
+				}
+				if (NeedCameraChange)
+				{
+					FLevelEditorViewportClient& LevelViewportClient = LevelViewport->GetLevelViewportClient();
+					if (LevelViewportClient.IsAnyActorLocked())
+					{
+						LevelViewportClient.SetActorLock(nullptr);
+						LevelViewportClient.bDisableInput = false;
+						LevelViewportClient.bEnableFading = true;
+					}
+				}
 			}
 		}
 #endif
 	};
+
+private:
+	bool NeedGameViewToggle;
+	bool NeedCameraChange;
+
 };
 
 UAutomationBlueprintFunctionLibrary::UAutomationBlueprintFunctionLibrary(const class FObjectInitializer& Initializer)
@@ -1200,7 +1220,7 @@ void UAutomationBlueprintFunctionLibrary::AutomationWaitForLoading(UObject* Worl
 	}
 }
 
-UAutomationEditorTask* UAutomationBlueprintFunctionLibrary::TakeHighResScreenshot(int32 ResX, int32 ResY, FString Filename, ACameraActor* Camera, bool bMaskEnabled, bool bCaptureHDR, EComparisonTolerance ComparisonTolerance, FString ComparisonNotes, float Delay)
+UAutomationEditorTask* UAutomationBlueprintFunctionLibrary::TakeHighResScreenshot(int32 ResX, int32 ResY, FString Filename, ACameraActor* Camera, bool bMaskEnabled, bool bCaptureHDR, EComparisonTolerance ComparisonTolerance, FString ComparisonNotes, float Delay, bool bForceGameView)
 {
 	UAutomationEditorTask* Task = NewObject<UAutomationEditorTask>();
 	FGCObjectScopeGuard TaskGuard(Task);
@@ -1212,13 +1232,15 @@ UAutomationEditorTask* UAutomationBlueprintFunctionLibrary::TakeHighResScreensho
 		{
 			FLevelEditorModule& LevelEditor = FModuleManager::GetModuleChecked<FLevelEditorModule>("LevelEditor");
 			SLevelViewport* LevelViewport = LevelEditor.GetFirstActiveLevelViewport().Get();
-			if (!LevelViewport->IsInGameView() && LevelViewport->CanToggleGameView())
+			bool bNeedGameViewToggle = bForceGameView && !LevelViewport->IsInGameView();
+			if (bNeedGameViewToggle && LevelViewport->CanToggleGameView())
 			{
 				LevelViewport->ToggleGameView();
 			}
 
 			// Move Viewport to Camera
-			if (Camera)
+			bool bNeedCameraChange = Camera != nullptr;
+			if (bNeedCameraChange)
 			{
 				FLevelEditorViewportClient& LevelViewportClient = LevelViewport->GetLevelViewportClient();
 				// We set the actor lock (pilot mode) and force the viewport to match the camera now.
@@ -1231,7 +1253,7 @@ UAutomationEditorTask* UAutomationBlueprintFunctionLibrary::TakeHighResScreensho
 
 			FinishLoadingBeforeScreenshot();
 
-			Task->BindTask(MakeUnique<FScreenshotTakenState>());
+			Task->BindTask(MakeUnique<FScreenshotTakenState>(bNeedGameViewToggle, bNeedCameraChange));
 
 			// Delay taking the screenshot by a few frames
 			FTSTicker::GetCoreTicker().AddTicker(TEXT("ScreenshotDelay"), Delay, [LevelViewport, ComparisonTolerance, ComparisonNotes, Filename, ResX, ResY, bMaskEnabled, bCaptureHDR](float) {
diff --git a/Engine/Source/Developer/FunctionalTesting/Public/AutomationBlueprintFunctionLibrary.h b/Engine/Source/Developer/FunctionalTesting/Public/AutomationBlueprintFunctionLibrary.h
--- a/Engine/Source/Developer/FunctionalTesting/Public/AutomationBlueprintFunctionLibrary.h
+++ b/Engine/Source/Developer/FunctionalTesting/Public/AutomationBlueprintFunctionLibrary.h
@@ -137,8 +137,8 @@ public:
 	/**
 	* take high res screenshot in editor.
 	*/
-	UFUNCTION(BlueprintCallable, Category = "Automation", meta = (AdvancedDisplay="Camera, bMaskEnabled, bCaptureHDR, ComparisonTolerance, ComparisonNotes"))
-	static UAutomationEditorTask* TakeHighResScreenshot(int32 ResX, int32 ResY, FString Filename, ACameraActor* Camera = nullptr, bool bMaskEnabled = false, bool bCaptureHDR = false, EComparisonTolerance ComparisonTolerance = EComparisonTolerance::Low, FString ComparisonNotes = TEXT(""), float Delay = 0.0);
+	UFUNCTION(BlueprintCallable, Category = "Automation", meta = (AdvancedDisplay="Camera, bMaskEnabled, bCaptureHDR, ComparisonTolerance, ComparisonNotes, bForceGameView"))
+	static UAutomationEditorTask* TakeHighResScreenshot(int32 ResX, int32 ResY, FString Filename, ACameraActor* Camera = nullptr, bool bMaskEnabled = false, bool bCaptureHDR = false, EComparisonTolerance ComparisonTolerance = EComparisonTolerance::Low, FString ComparisonNotes = TEXT(""), float Delay = 0.0, bool bForceGameView = true);
 
 	/**
 	* request image comparison.
