Adding flag to toggle whether or not garbage collection should be delayed, by default they are. Added property to show in the Editor Details panel for the Functional Test in the Advanced section.

PR # 10898
#commit 10898

diff --git a/Engine/Source/Developer/FunctionalTesting/Classes/FunctionalTest.h b/Engine/Source/Developer/FunctionalTesting/Classes/FunctionalTest.h
--- a/Engine/Source/Developer/FunctionalTesting/Classes/FunctionalTest.h
+++ b/Engine/Source/Developer/FunctionalTesting/Classes/FunctionalTest.h
@@ -317,6 +317,12 @@ protected:
 	UPROPERTY(EditAnywhere, BlueprintReadWrite, Category="Functional Testing")
 	TObjectPtr<AActor> ObservationPoint;
 
+	/**
+	 * Allows for garbage collection to be delayed. If delayed, garbage collection will be triggered at the end of a test run
+	 */
+	UPROPERTY(EditAnywhere, BlueprintReadOnly, Category="Functional Testing", AdvancedDisplay)
+	uint32 bShouldDelayGarbageCollection:1;
+
 	/**
 	 * A random number stream that you can use during testing.  This number stream will be consistent
 	 * every time the test is run.
diff --git a/Engine/Source/Developer/FunctionalTesting/Private/FunctionalTest.cpp b/Engine/Source/Developer/FunctionalTesting/Private/FunctionalTest.cpp
--- a/Engine/Source/Developer/FunctionalTesting/Private/FunctionalTest.cpp
+++ b/Engine/Source/Developer/FunctionalTesting/Private/FunctionalTest.cpp
@@ -122,6 +122,7 @@ AFunctionalTest::AFunctionalTest( const FObjectInitializer& ObjectInitializer )
 	, bIsEnabled(true)
 	, LogErrorHandling(EFunctionalTestLogHandling::ProjectDefault)
 	, LogWarningHandling(EFunctionalTestLogHandling::ProjectDefault)
+	, bShouldDelayGarbageCollection(true)
 	, Result(EFunctionalTestResult::Invalid)
 	, PreparationTimeLimit(15.0f)
 	, TimeLimit(60.0f)
@@ -332,8 +333,12 @@ void AFunctionalTest::Tick(float DeltaSeconds)
 	}
 	SCOPE_CYCLE_COUNTER(STAT_FunctionalTest_TickTest);
 	
-	//Do not collect garbage during the test. We force GC at the end.
-	GEngine->DelayGarbageCollection();
+	//Allow Functional Tests to configure if GC is delayed until the end. 
+	if (bShouldDelayGarbageCollection)
+	{
+		//Do not collect garbage during the test. We force GC at the end.
+		GEngine->DelayGarbageCollection();
+	}
 
 	TotalTime += DeltaSeconds;
 
