PR #10982: Disable and hide `HLOD` category for `Functional Test` Actor
Reduce options bloat from new feature(HLOD).

diff --git a/Engine/Source/Developer/FunctionalTesting/Classes/FunctionalTest.h b/Engine/Source/Developer/FunctionalTesting/Classes/FunctionalTest.h
--- a/Engine/Source/Developer/FunctionalTesting/Classes/FunctionalTest.h
+++ b/Engine/Source/Developer/FunctionalTesting/Classes/FunctionalTest.h
@@ -260,7 +260,7 @@ private:
 DECLARE_DYNAMIC_MULTICAST_DELEGATE(FFunctionalTestEventSignature);
 DECLARE_DELEGATE_OneParam(FFunctionalTestDoneSignature, class AFunctionalTest*);
 
-UCLASS(hidecategories=( Actor, Input, Rendering ), Blueprintable)
+UCLASS(hidecategories=( Actor, Input, Rendering, HLOD ), Blueprintable)
 class FUNCTIONALTESTING_API AFunctionalTest : public AActor
 {
 	GENERATED_BODY()
diff --git a/Engine/Source/Developer/FunctionalTesting/Private/FunctionalTest.cpp b/Engine/Source/Developer/FunctionalTesting/Private/FunctionalTest.cpp
--- a/Engine/Source/Developer/FunctionalTesting/Private/FunctionalTest.cpp
+++ b/Engine/Source/Developer/FunctionalTesting/Private/FunctionalTest.cpp
@@ -141,6 +141,7 @@ AFunctionalTest::AFunctionalTest( const FObjectInitializer& ObjectInitializer )
 	PrimaryActorTick.bTickEvenWhenPaused = true;
 	
 	SetCanBeDamaged(false);
+	bEnableAutoLODGeneration = false;
 
 	SpriteComponent = CreateDefaultSubobject<UBillboardComponent>(TEXT("Sprite"));
 	if (SpriteComponent)
