Add the ViewRectOffset "test variant" to all screenshot functional tests. When a test variant is enabled, the test will take two screenshots, one in the base configuration, and one in the variant configuration, and fail if both are not similar to the ground truth. The ViewRectOffset variant will be enabled if "Support Stereo Test Variants" is checked on the functional test actor (TRUE by default) and the CVar r.EnableStereoScreenshotTestVariants is set to 1. The CVar is 0 by default for the time being, so there should be no changes to ET unless it is manually set to true.

diff --git a/Engine/Source/Developer/FunctionalTesting/Classes/ScreenshotFunctionalTest.h b/Engine/Source/Developer/FunctionalTesting/Classes/ScreenshotFunctionalTest.h
--- a/Engine/Source/Developer/FunctionalTesting/Classes/ScreenshotFunctionalTest.h
+++ b/Engine/Source/Developer/FunctionalTesting/Classes/ScreenshotFunctionalTest.h
@@ -28,7 +28,21 @@ public:
 	UPROPERTY(EditAnywhere, BlueprintReadWrite, Category = "Camera", SimpleDisplay)
 	bool bCameraCutOnScreenshotPrep;
 
+	// If r.EnableStereoScreenshotTestVariants is enabled, also take and compare screenshots in stereo rendering emulation modes when running this test
+	UPROPERTY(EditAnywhere, BlueprintReadWrite, Category = "Functional Testing", SimpleDisplay)
+	bool bSupportStereoTestVariants = true;
+
 protected:
 	virtual void PrepareTest() override;
 	virtual void RequestScreenshot() override;
+	virtual void OnScreenShotCaptured(int32 InSizeX, int32 InSizeY, const TArray<FColor>& InImageData) override;
+	virtual void OnScreenshotTakenAndCompared() override;
+
+private:
+	void PerformVariant(FString VariantName, FString SetpCommand, FString RestoreCommand);
+
+	FString CurrentVariantName;
+	FString VariantRestoreCommand;
+	bool bNeedsVariantRestore;
+	bool bShouldDoViewRectOffsetVariant;
 };
diff --git a/Engine/Source/Developer/FunctionalTesting/Classes/ScreenshotFunctionalTestBase.h b/Engine/Source/Developer/FunctionalTesting/Classes/ScreenshotFunctionalTestBase.h
--- a/Engine/Source/Developer/FunctionalTesting/Classes/ScreenshotFunctionalTestBase.h
+++ b/Engine/Source/Developer/FunctionalTesting/Classes/ScreenshotFunctionalTestBase.h
@@ -10,6 +10,8 @@
 
 #include "ScreenshotFunctionalTestBase.generated.h"
 
+DEFINE_LOG_CATEGORY_STATIC(LogScreenshotFunctionalTest, Log, Log)
+
 class FAutomationTestScreenshotEnvSetup;
 
 /**
@@ -52,7 +54,7 @@ protected:
 	// Pass screenshot pixels and meta data to FAutomationTestFramework. Register
 	// OnComparisonComplete which will be called the automation test system when
 	// screenshot comparison is complete
-	void OnScreenShotCaptured(int32 InSizeX, int32 InSizeY, const TArray<FColor>& InImageData);
+	virtual void OnScreenShotCaptured(int32 InSizeX, int32 InSizeY, const TArray<FColor>& InImageData);
 
 	// Do some logging and trigger OnScreenshotTakenAndCompared
 	void OnComparisonComplete(const FAutomationScreenshotCompareResults& CompareResults);
diff --git a/Engine/Source/Developer/FunctionalTesting/Private/ScreenshotFunctionalTest.cpp b/Engine/Source/Developer/FunctionalTesting/Private/ScreenshotFunctionalTest.cpp
--- a/Engine/Source/Developer/FunctionalTesting/Private/ScreenshotFunctionalTest.cpp
+++ b/Engine/Source/Developer/FunctionalTesting/Private/ScreenshotFunctionalTest.cpp
@@ -19,9 +19,16 @@
 
 #include UE_INLINE_GENERATED_CPP_BY_NAME(ScreenshotFunctionalTest)
 
+static TAutoConsoleVariable<int32> CVarEnableStereoTestVariants(
+	TEXT("r.EnableStereoScreenshotTestVariants"), 0,
+	TEXT("Allows screenshot comparison tests with \"Support Stereo Test Variants\" checked to also test stereo rendering."),
+	ECVF_Default);
+
 AScreenshotFunctionalTest::AScreenshotFunctionalTest( const FObjectInitializer& ObjectInitializer )
 	: AScreenshotFunctionalTestBase(ObjectInitializer)
 	, bCameraCutOnScreenshotPrep(true)
+	, bNeedsVariantRestore(false)
+	, bShouldDoViewRectOffsetVariant(false)
 {
 }
 
@@ -39,6 +46,8 @@ void AScreenshotFunctionalTest::Serialize(FArchive& Ar)
 
 void AScreenshotFunctionalTest::PrepareTest()
 {
+	bShouldDoViewRectOffsetVariant = bSupportStereoTestVariants && CVarEnableStereoTestVariants.GetValueOnAnyThread() > 0;
+
 	// Pre-prep flush to allow rendering to temporary targets and other test resources
 	UAutomationBlueprintFunctionLibrary::FinishLoadingBeforeScreenshot();
 
@@ -97,3 +106,101 @@ void AScreenshotFunctionalTest::RequestScreenshot()
 	}
 }
 
+void AScreenshotFunctionalTest::OnScreenShotCaptured(int32 InSizeX, int32 InSizeY, const TArray<FColor>& InImageData)
+{
+	UGameViewportClient* GameViewportClient = AutomationCommon::GetAnyGameViewportClient();
+	check(GameViewportClient);
+
+	GameViewportClient->OnScreenshotCaptured().RemoveAll(this);
+
+#if WITH_AUTOMATION_TESTS
+	const FString Context = AutomationCommon::GetWorldContext(GetWorld());
+
+	TArray<uint8> CapturedFrameTrace = AutomationCommon::CaptureFrameTrace(Context, TestLabel);
+
+	FAutomationScreenshotData Data = UAutomationBlueprintFunctionLibrary::BuildScreenshotData(Context, TestLabel, InSizeX, InSizeY);
+
+	// Copy the relevant data into the metadata for the screenshot.
+	Data.bHasComparisonRules = true;
+	Data.ToleranceRed = ScreenshotOptions.ToleranceAmount.Red;
+	Data.ToleranceGreen = ScreenshotOptions.ToleranceAmount.Green;
+	Data.ToleranceBlue = ScreenshotOptions.ToleranceAmount.Blue;
+	Data.ToleranceAlpha = ScreenshotOptions.ToleranceAmount.Alpha;
+	Data.ToleranceMinBrightness = ScreenshotOptions.ToleranceAmount.MinBrightness;
+	Data.ToleranceMaxBrightness = ScreenshotOptions.ToleranceAmount.MaxBrightness;
+	Data.bIgnoreAntiAliasing = ScreenshotOptions.bIgnoreAntiAliasing;
+	Data.bIgnoreColors = ScreenshotOptions.bIgnoreColors;
+	Data.MaximumLocalError = ScreenshotOptions.MaximumLocalError;
+	Data.MaximumGlobalError = ScreenshotOptions.MaximumGlobalError;
+
+	// Add the notes
+	Data.Notes = Notes;
+
+	// If variant in use, pass on the name, then restore settings since capture is done
+	Data.VariantName = CurrentVariantName;
+	CurrentVariantName = "";
+
+	if (bNeedsVariantRestore)
+	{
+		GEngine->Exec(nullptr, *VariantRestoreCommand);
+		bNeedsVariantRestore = false;
+	}
+
+	if (GIsAutomationTesting)
+	{
+		FAutomationTestFramework::Get().OnScreenshotCompared.AddUObject(this, &AScreenshotFunctionalTest::OnComparisonComplete);
+	}
+
+	FAutomationTestFramework::Get().OnScreenshotAndTraceCaptured().ExecuteIfBound(InImageData, CapturedFrameTrace, Data);
+
+	UE_LOG(LogScreenshotFunctionalTest, Log, TEXT("Screenshot captured as %s"), *Data.ScreenshotName);
+#endif
+}
+
+void AScreenshotFunctionalTest::OnScreenshotTakenAndCompared()
+{
+	if (bShouldDoViewRectOffsetVariant)
+	{
+		bShouldDoViewRectOffsetVariant = false;
+
+		// 0: Off, 1: Center, 2: Top Left, 3: Top Right, 4: Bottom Left, 5: Bottom Right
+		// We use bottom right to test both horizontal and vertical offsets (typically used for XR and split-screen, respectively).
+		PerformVariant("ViewRectOffset", "r.Test.ViewRectOffset 5", "r.Test.ViewRectOffset 0");
+	}
+	else
+	{
+		// This ends the test and reports results
+		Super::OnScreenshotTakenAndCompared();
+	}
+}
+
+void AScreenshotFunctionalTest::PerformVariant(FString VariantName, FString SetupCommand, FString RestoreCommand)
+{
+	// Set up variant
+	GEngine->Exec(nullptr, *SetupCommand);
+
+	// Save variant name and command needed to restore in OnScreenShotCaptured
+	CurrentVariantName = VariantName;
+	if (!RestoreCommand.IsEmpty())
+	{
+		VariantRestoreCommand = RestoreCommand;
+		bNeedsVariantRestore = true;
+	}
+
+	// Re-prepare test and request screenshot comparison
+	if (bCameraCutOnScreenshotPrep)
+	{
+		APlayerController* PlayerController = UGameplayStatics::GetPlayerController(GetWorld(), 0);
+
+		if (PlayerController && PlayerController->PlayerCameraManager)
+		{
+			PlayerController->PlayerCameraManager->SetGameCameraCutThisFrame();
+			if (ScreenshotCamera)
+			{
+				ScreenshotCamera->NotifyCameraCut();
+			}
+		}
+	}
+
+	RequestScreenshot(); 
+}
diff --git a/Engine/Source/Developer/FunctionalTesting/Private/ScreenshotFunctionalTestBase.cpp b/Engine/Source/Developer/FunctionalTesting/Private/ScreenshotFunctionalTestBase.cpp
--- a/Engine/Source/Developer/FunctionalTesting/Private/ScreenshotFunctionalTestBase.cpp
+++ b/Engine/Source/Developer/FunctionalTesting/Private/ScreenshotFunctionalTestBase.cpp
@@ -24,8 +24,6 @@
 
 #define	WITH_EDITOR_AUTOMATION_TESTS	(WITH_EDITOR && WITH_AUTOMATION_TESTS)
 
-DEFINE_LOG_CATEGORY_STATIC(LogScreenshotFunctionalTest, Log, Log)
-
 static TAutoConsoleVariable<int32> GDumpGPUDumpOnScreenshotTest(
 	TEXT("r.DumpGPU.DumpOnScreenshotTest"), 0,
 	TEXT("Allows to filter the tree when using r.DumpGPU command, the pattern match is case sensitive."),
diff --git a/Engine/Source/Developer/ScreenShotComparison/Private/Models/ScreenComparisonModel.cpp b/Engine/Source/Developer/ScreenShotComparison/Private/Models/ScreenComparisonModel.cpp
--- a/Engine/Source/Developer/ScreenShotComparison/Private/Models/ScreenComparisonModel.cpp
+++ b/Engine/Source/Developer/ScreenShotComparison/Private/Models/ScreenComparisonModel.cpp
@@ -85,11 +85,12 @@ FString FScreenComparisonModel::GetName()
 		auto LoadedMetadata = GetMetadata();
 		if (LoadedMetadata.IsSet())
 		{
-			FString NameString = FString::Printf(TEXT("%s.%s"), *LoadedMetadata->Context, *LoadedMetadata->ScreenShotName);
+			FString VariantSuffix = LoadedMetadata->VariantName.Len() > 0 ? FString::Printf(TEXT(".%s"), *LoadedMetadata->VariantName) : TEXT("");
+			FString NameString = FString::Printf(TEXT("%s.%s%s"), *LoadedMetadata->Context, *LoadedMetadata->ScreenShotName, *VariantSuffix);
 			if ((LoadedMetadata->Context.Len() && LoadedMetadata->TestName.Len())
 				|| !LoadedMetadata->ScreenShotName.Len())
 			{
-				NameString = FString::Printf(TEXT("%s.%s"), *LoadedMetadata->Context, *LoadedMetadata->TestName);
+				NameString = FString::Printf(TEXT("%s.%s%s"), *LoadedMetadata->Context, *LoadedMetadata->TestName, *VariantSuffix);
 			}
 
 			Name = NameString;
diff --git a/Engine/Source/Developer/ScreenShotComparisonTools/Private/ScreenShotManager.cpp b/Engine/Source/Developer/ScreenShotComparisonTools/Private/ScreenShotManager.cpp
--- a/Engine/Source/Developer/ScreenShotComparisonTools/Private/ScreenShotManager.cpp
+++ b/Engine/Source/Developer/ScreenShotComparisonTools/Private/ScreenShotManager.cpp
@@ -80,9 +80,16 @@ FString FScreenShotManager::GetPathComponentForPlatformAndRHI(const FAutomationS
 /**
  * Images are now preferred to be under MapOrContext/ImageName/Plat/RHI etc
  */
-FString FScreenShotManager::GetPathComponentForTestImages(const FAutomationScreenshotMetadata& MetaData) const
+FString FScreenShotManager::GetPathComponentForTestImages(const FAutomationScreenshotMetadata& MetaData, bool bIncludeVariantName) const
 {
-	return FPaths::Combine(MetaData.Context, *MetaData.ScreenShotName);
+	if (bIncludeVariantName)
+	{
+		return FPaths::Combine(MetaData.Context, *MetaData.ScreenShotName, *MetaData.VariantName);
+	}
+	else
+	{
+		return FPaths::Combine(MetaData.Context, *MetaData.ScreenShotName);
+	}
 }
 
 FString FScreenShotManager::GetApprovedFolderForImageWithOptions(const FAutomationScreenshotMetadata& MetaData, EApprovedFolderOptions InOptions) const
@@ -92,7 +99,7 @@ FString FScreenShotManager::GetApprovedFolderForImageWithOptions(const FAutomati
 	bool bUsePlatformPath = PlatInfo.bIsConfidential && (InOptions & EApprovedFolderOptions::UsePlatformFolders) == 0;
 
 	// Test folder will be MapOrContext/ImageName
-	FString TestFolder = GetPathComponentForTestImages(MetaData);
+	FString TestFolder = GetPathComponentForTestImages(MetaData, false);
 
 	FString OutPath = FPaths::ProjectDir();
 
@@ -280,7 +287,7 @@ FImageComparisonResult FScreenShotManager::CompareScreenshot(const FString& InUn
 	// get the ideal path for our approved image. This is the path that a file would be at if it matches our platform and RHI
 	FString IdealApprovedFolderPath = GetIdealApprovedFolderForImage(IncomingMetaData);
 
-	FString ResultsSubFolder = GetPathComponentForTestImages(IncomingMetaData);
+	FString ResultsSubFolder = GetPathComponentForTestImages(IncomingMetaData, true);
 
 	// If the metadata for the screenshot does not provide tolerance rules, use these instead.
 	FImageTolerance DefaultTolerance = FImageTolerance::DefaultIgnoreLess;
@@ -658,8 +665,8 @@ TSharedPtr<FImageComparisonResult> FScreenShotManager::CompareImageSequence(cons
 			ComparisonResult.SourcePlatform = Metadata.Platform;
 			ComparisonResult.SourceRHI = Metadata.Rhi;
 			ComparisonResult.IdealApprovedFolderPath = GetIdealApprovedFolderForImage(Metadata);;
-			ComparisonResult.ScreenshotName = GetPathComponentForTestImages(Metadata);
-			FString ReportPathOnDisk = FPaths::Combine(ScreenshotResultsFolder, GetPathComponentForTestImages(Metadata), Metadata.Platform, GetPathComponentForRHI(Metadata), TEXT("/"));
+			ComparisonResult.ScreenshotName = GetPathComponentForTestImages(Metadata, true);
+			FString ReportPathOnDisk = FPaths::Combine(ScreenshotResultsFolder, GetPathComponentForTestImages(Metadata, true), Metadata.Platform, GetPathComponentForRHI(Metadata), TEXT("/"));
 
 			ComparisonResult.ReportApprovedFilePath = FPaths::Combine(ReportPathOnDisk, TEXT("Approved.png"));
 			ComparisonResult.ReportIncomingFilePath = FPaths::Combine(ReportPathOnDisk, TEXT("Incoming.png"));
diff --git a/Engine/Source/Developer/ScreenShotComparisonTools/Private/ScreenShotManager.h b/Engine/Source/Developer/ScreenShotComparisonTools/Private/ScreenShotManager.h
--- a/Engine/Source/Developer/ScreenShotComparisonTools/Private/ScreenShotManager.h
+++ b/Engine/Source/Developer/ScreenShotComparisonTools/Private/ScreenShotManager.h
@@ -64,7 +64,7 @@ private:
 
 	FString	GetPathComponentForRHI(const FAutomationScreenshotMetadata& MetaData) const;
 	FString	GetPathComponentForPlatformAndRHI(const FAutomationScreenshotMetadata& MetaData) const;
-	FString GetPathComponentForTestImages(const FAutomationScreenshotMetadata& MetaData) const;
+	FString GetPathComponentForTestImages(const FAutomationScreenshotMetadata& MetaData, bool bIncludeVariantName) const;
 		
 
 	FString GetApprovedFolderForImageWithOptions(const FAutomationScreenshotMetadata& MetaData, EApprovedFolderOptions InOptions) const;
diff --git a/Engine/Source/Runtime/AutomationMessages/Public/AutomationWorkerMessages.h b/Engine/Source/Runtime/AutomationMessages/Public/AutomationWorkerMessages.h
--- a/Engine/Source/Runtime/AutomationMessages/Public/AutomationWorkerMessages.h
+++ b/Engine/Source/Runtime/AutomationMessages/Public/AutomationWorkerMessages.h
@@ -402,6 +402,9 @@ public:
 	UPROPERTY(EditAnywhere, Category="Message")
 	FString ScreenShotName;
 
+	UPROPERTY(EditAnywhere, Category = "Message")
+	FString VariantName;
+
 	UPROPERTY(EditAnywhere, Category="Message")
 	FString Context;
 
@@ -528,6 +531,7 @@ public:
 
 		// Human readable name and associated context the screenshot was taken in.
 		ScreenShotName = Data.ScreenShotName;
+		VariantName = Data.VariantName;
 		Context = Data.Context;
 		TestName = Data.TestName;
 		Notes = Data.Notes;
diff --git a/Engine/Source/Runtime/Core/Public/Misc/AutomationTest.h b/Engine/Source/Runtime/Core/Public/Misc/AutomationTest.h
--- a/Engine/Source/Runtime/Core/Public/Misc/AutomationTest.h
+++ b/Engine/Source/Runtime/Core/Public/Misc/AutomationTest.h
@@ -698,6 +698,7 @@ FORCEINLINE uint32 GetTypeHash(const FAutomationExpectedMessage& Object)
 struct FAutomationScreenshotData
 {
 	FString ScreenShotName;
+	FString VariantName;
 	FString Context;
 	FString TestName;
 	FString Notes;
