Support the capability to add info, warning and error directly from Python test

diff --git a/Engine/Source/Developer/FunctionalTesting/Private/AutomationBlueprintFunctionLibrary.cpp b/Engine/Source/Developer/FunctionalTesting/Private/AutomationBlueprintFunctionLibrary.cpp
--- a/Engine/Source/Developer/FunctionalTesting/Private/AutomationBlueprintFunctionLibrary.cpp
+++ b/Engine/Source/Developer/FunctionalTesting/Private/AutomationBlueprintFunctionLibrary.cpp
@@ -1502,5 +1502,29 @@ void UAutomationBlueprintFunctionLibrary::SetEditorViewportVisualizeBuffer( FNam
 #endif
 }
 
+void UAutomationBlueprintFunctionLibrary::AddTestInfo(const FString& InLogItem)
+{
+	if (FAutomationTestBase* CurrentTest = FAutomationTestFramework::Get().GetCurrentTest())
+	{
+		CurrentTest->AddInfo(InLogItem);
+	}
+}
+
+void UAutomationBlueprintFunctionLibrary::AddTestWarning(const FString& InLogItem)
+{
+	if (FAutomationTestBase* CurrentTest = FAutomationTestFramework::Get().GetCurrentTest())
+	{
+		CurrentTest->AddWarning(InLogItem);
+	}
+}
+
+void UAutomationBlueprintFunctionLibrary::AddTestError(const FString& InLogItem)
+{
+	if (FAutomationTestBase* CurrentTest = FAutomationTestFramework::Get().GetCurrentTest())
+	{
+		CurrentTest->AddError(InLogItem);
+	}
+}
+
 #undef LOCTEXT_NAMESPACE
 
diff --git a/Engine/Source/Developer/FunctionalTesting/Public/AutomationBlueprintFunctionLibrary.h b/Engine/Source/Developer/FunctionalTesting/Public/AutomationBlueprintFunctionLibrary.h
--- a/Engine/Source/Developer/FunctionalTesting/Public/AutomationBlueprintFunctionLibrary.h
+++ b/Engine/Source/Developer/FunctionalTesting/Public/AutomationBlueprintFunctionLibrary.h
@@ -218,6 +218,24 @@ public:
 	/** Sets all viewports of the first found level editor to have the VisualizeBuffer ViewMode and also display a given buffer (BaseColor/Metallic/Roughness/etc.) **/
 	UFUNCTION(BlueprintCallable, Category = "Automation")
 	static void SetEditorViewportVisualizeBuffer(FName BufferName);
+
+	/**
+	 * Add info to currently running automated test.
+	 */
+	UFUNCTION(BlueprintCallable, Category = "Automation")
+	static void AddTestInfo(const FString& InLogItem);
+
+	/**
+	 * Add warning to currently running automated test.
+	 */
+	UFUNCTION(BlueprintCallable, Category = "Automation")
+	static void AddTestWarning(const FString& InLogItem);
+
+	/**
+	 * Add error to currently running automated test.
+	 */
+	UFUNCTION(BlueprintCallable, Category = "Automation")
+	static void AddTestError(const FString& InLogItem);
 };
 
 #if WITH_AUTOMATION_TESTS
