Fix linker issues when declaring types outside of "FunctionalTesting" that inherit from "FFunctionalTestBase"

diff --git a/Engine/Source/Developer/FunctionalTesting/Classes/FunctionalTestBase.h b/Engine/Source/Developer/FunctionalTesting/Classes/FunctionalTestBase.h
--- a/Engine/Source/Developer/FunctionalTesting/Classes/FunctionalTestBase.h
+++ b/Engine/Source/Developer/FunctionalTesting/Classes/FunctionalTestBase.h
@@ -45,7 +45,7 @@ public:
 		{
 			bElevateLogWarningsToErrors = InWarningsAreErrors.GetValue();
 		}
-	}	
+	}
 
 	/**
 	 * Determines if Error logs should be suppressed from test results
@@ -84,22 +84,24 @@ public:
 	/**
 	 * Returns the name of the running functional test. Empty if no test is running
 	 */
-	static FString GetRunningTestName() { return ActiveTestName; }
+	FUNCTIONALTESTING_API static FString GetRunningTestName() { return ActiveTestName; }
 
 	/**
 	 * Returns true if a functional test is running (does not include map setup)
 	 */
-	static bool IsFunctionalTestRunning()	{ return bIsFunctionalTestRunning;	}
+	FUNCTIONALTESTING_API static bool IsFunctionalTestRunning()	{ return bIsFunctionalTestRunning;	}
 
 protected:
 
-	void SetLogErrorAndWarningHandlingToDefault();
+	FUNCTIONALTESTING_API void SetLogErrorAndWarningHandlingToDefault();
 
 	bool bSuppressLogErrors;
 	bool bSuppressLogWarnings;
 	bool bElevateLogWarningsToErrors;
 	bool bSuppressLogs;
 
+private:
+
 	static bool bIsFunctionalTestRunning;
 	static FString ActiveTestName;
 };
