Update LightweightStereoTestVariants mode to run ONLY stereo test variants, no baseline, since they've been moved into their own category. Remove support for specifically exempting tests from running stereo variants. Save screenshots even for passed variants for the time being, since it's now possible to have a variant be the only test running.

Now `Automation.EnableStereoTestVariants=1` (off by default) enables ViewRectOffset, and adding `Automation.LightweightStereoTestVariants=1` causes us to ONLY run ViewRectOffset (on by default).

diff --git a/Engine/Source/Developer/FunctionalTesting/Classes/ScreenshotFunctionalTest.h b/Engine/Source/Developer/FunctionalTesting/Classes/ScreenshotFunctionalTest.h
--- a/Engine/Source/Developer/FunctionalTesting/Classes/ScreenshotFunctionalTest.h
+++ b/Engine/Source/Developer/FunctionalTesting/Classes/ScreenshotFunctionalTest.h
@@ -28,21 +28,20 @@ public:
 	UPROPERTY(EditAnywhere, BlueprintReadWrite, Category = "Camera", SimpleDisplay)
 	bool bCameraCutOnScreenshotPrep;
 
-	// If r.EnableStereoScreenshotTestVariants is enabled, also take and compare screenshots in stereo rendering emulation modes when running this test
-	UPROPERTY(EditAnywhere, BlueprintReadWrite, Category = "Screenshot", SimpleDisplay)
-	bool bSupportStereoTestVariants = true;
-
 protected:
 	virtual void PrepareTest() override;
 	virtual void RequestScreenshot() override;
 	virtual void OnScreenShotCaptured(int32 InSizeX, int32 InSizeY, const TArray<FColor>& InImageData) override;
 	virtual void OnScreenshotTakenAndCompared() override;
+	virtual void StartTest() override;
 
 private:
-	void PerformVariant(FString VariantName, FString SetpCommand, FString RestoreCommand);
+	void SetupVariant(FString VariantName, FString SetpCommand, FString RestoreCommand);
 
 	FString CurrentVariantName;
 	FString VariantRestoreCommand;
 	bool bNeedsVariantRestore;
+
+	bool bShouldDoBaselineTest;
 	bool bShouldDoViewRectOffsetVariant;
 };
diff --git a/Engine/Source/Developer/FunctionalTesting/Private/ScreenshotFunctionalTest.cpp b/Engine/Source/Developer/FunctionalTesting/Private/ScreenshotFunctionalTest.cpp
--- a/Engine/Source/Developer/FunctionalTesting/Private/ScreenshotFunctionalTest.cpp
+++ b/Engine/Source/Developer/FunctionalTesting/Private/ScreenshotFunctionalTest.cpp
@@ -23,6 +23,7 @@ AScreenshotFunctionalTest::AScreenshotFunctionalTest( const FObjectInitializer&
 	: AScreenshotFunctionalTestBase(ObjectInitializer)
 	, bCameraCutOnScreenshotPrep(true)
 	, bNeedsVariantRestore(false)
+	, bShouldDoBaselineTest(false)
 	, bShouldDoViewRectOffsetVariant(false)
 {
 }
@@ -41,7 +42,9 @@ void AScreenshotFunctionalTest::Serialize(FArchive& Ar)
 
 void AScreenshotFunctionalTest::PrepareTest()
 {
-	bShouldDoViewRectOffsetVariant = bSupportStereoTestVariants && FAutomationTestFramework::NeedPerformStereoTestVariants();
+	// If variants are enabled and lightweight variants are on, skip the baseline test
+	bShouldDoViewRectOffsetVariant = FAutomationTestFramework::NeedPerformStereoTestVariants();
+	bShouldDoBaselineTest = !(bShouldDoViewRectOffsetVariant && FAutomationTestFramework::NeedUseLightweightStereoTestVariants());
 
 	// Pre-prep flush to allow rendering to temporary targets and other test resources
 	UAutomationBlueprintFunctionLibrary::FinishLoadingBeforeScreenshot();
@@ -152,19 +155,47 @@ void AScreenshotFunctionalTest::OnScreenShotCaptured(int32 InSizeX, int32 InSize
 #endif
 }
 
+void AScreenshotFunctionalTest::StartTest()
+{
+	if (bShouldDoBaselineTest)
+	{
+		bShouldDoBaselineTest = false;
+	}
+	else if (bShouldDoViewRectOffsetVariant)
+	{
+		bShouldDoViewRectOffsetVariant = false;
+		SetupVariant("ViewRectOffset", "r.Test.ViewRectOffset 5", "r.Test.ViewRectOffset 0");
+	}
+
+	Super::StartTest();
+}
+
 void AScreenshotFunctionalTest::OnScreenshotTakenAndCompared()
 {
 	FAutomationTestBase* CurrentTest = FAutomationTestFramework::Get().GetCurrentTest();
 	bool bSkipDueToError = FAutomationTestFramework::Get().NeedUseLightweightStereoTestVariants() && (!CurrentTest || CurrentTest->HasAnyErrors());
 
+	// If we still need to perform any variants, loop back here
 	if (bShouldDoViewRectOffsetVariant && !bSkipDueToError)
 	{
-		bShouldDoViewRectOffsetVariant = false;
+		// Re-prepare test if necessary
+		if (bCameraCutOnScreenshotPrep)
+		{
+			APlayerController* PlayerController = UGameplayStatics::GetPlayerController(GetWorld(), 0);
+
+			if (PlayerController && PlayerController->PlayerCameraManager)
+			{
+				PlayerController->PlayerCameraManager->SetGameCameraCutThisFrame();
+				if (ScreenshotCamera)
+				{
+					ScreenshotCamera->NotifyCameraCut();
+				}
+			}
+		}
 
-		// 0: Off, 1: Center, 2: Top Left, 3: Top Right, 4: Bottom Left, 5: Bottom Right
-		// We use bottom right to test both horizontal and vertical offsets (typically used for XR and split-screen, respectively).
-		PerformVariant("ViewRectOffset", "r.Test.ViewRectOffset 5", "r.Test.ViewRectOffset 0");
+		StartTest();
 	}
+
 	else
 	{
 		// This ends the test and reports results
@@ -172,7 +203,7 @@ void AScreenshotFunctionalTest::OnScreenshotTakenAndCompared()
 	}
 }
 
-void AScreenshotFunctionalTest::PerformVariant(FString VariantName, FString SetupCommand, FString RestoreCommand)
+void AScreenshotFunctionalTest::SetupVariant(FString VariantName, FString SetupCommand, FString RestoreCommand)
 {
 	// Set up variant
 	GEngine->Exec(nullptr, *SetupCommand);
@@ -184,21 +215,4 @@ void AScreenshotFunctionalTest::PerformVariant(FString VariantName, FString Setu
 		VariantRestoreCommand = RestoreCommand;
 		bNeedsVariantRestore = true;
 	}
-
-	// Re-prepare test and request screenshot comparison
-	if (bCameraCutOnScreenshotPrep)
-	{
-		APlayerController* PlayerController = UGameplayStatics::GetPlayerController(GetWorld(), 0);
-
-		if (PlayerController && PlayerController->PlayerCameraManager)
-		{
-			PlayerController->PlayerCameraManager->SetGameCameraCutThisFrame();
-			if (ScreenshotCamera)
-			{
-				ScreenshotCamera->NotifyCameraCut();
-			}
-		}
-	}
-
-	RequestScreenshot(); 
 }
diff --git a/Engine/Source/Developer/ScreenShotComparisonTools/Private/ScreenShotManager.cpp b/Engine/Source/Developer/ScreenShotComparisonTools/Private/ScreenShotManager.cpp
--- a/Engine/Source/Developer/ScreenShotComparisonTools/Private/ScreenShotManager.cpp
+++ b/Engine/Source/Developer/ScreenShotComparisonTools/Private/ScreenShotManager.cpp
@@ -392,12 +392,13 @@ FImageComparisonResult FScreenShotManager::CompareScreenshot(const FString& InUn
 	ComparisonResult.ScreenshotName = ResultsSubFolder;
 
 	// Do not save passing variant test screenshots
-	const bool bIsVariant = !IncomingMetaData.VariantName.IsEmpty();
+	// Disabled for now until more variants are added, since variants now run without baseline tests in lightweight mode and we need to save at least one screenshot
+	/*const bool bIsVariant = !IncomingMetaData.VariantName.IsEmpty();
 	if (bIsVariant && FAutomationTestFramework::Get().NeedUseLightweightStereoTestVariants() && ComparisonResult.AreSimilar())
 	{
 		ComparisonResult.bSkipAttachingImages = true;
 		return ComparisonResult;
-	}
+	}*/
 
 	// Result paths should be relative to the project. Note this may be empty, and if it is MakePathRelative returns
 	// a non empty relative path... but we want it to stay empty as that's how we signal that no approved file exists
