FunctionalTesting - Make screenshot test actor details more intuitive while working with DisableNoisyRenderingFeatures/EyeAdaptation options.

#rnx

diff --git a/Engine/Source/Developer/FunctionalTesting/Private/AutomationBlueprintFunctionLibrary.cpp b/Engine/Source/Developer/FunctionalTesting/Private/AutomationBlueprintFunctionLibrary.cpp
--- a/Engine/Source/Developer/FunctionalTesting/Private/AutomationBlueprintFunctionLibrary.cpp
+++ b/Engine/Source/Developer/FunctionalTesting/Private/AutomationBlueprintFunctionLibrary.cpp
@@ -225,7 +225,7 @@ public:
 			//InViewFamily.EngineShowFlags.SetEyeAdaptation(false);
 			//InViewFamily.EngineShowFlags.SetTonemapper(false);
 		}
-		}
+	}
 
 	virtual void BeginRenderViewFamily(FSceneViewFamily& InViewFamily) {}
 	virtual void PreRenderViewFamily_RenderThread(FRDGBuilder& GraphBuilder, FSceneViewFamily& InViewFamily) {}
@@ -275,18 +275,16 @@ void FAutomationTestScreenshotEnvSetup::Setup(UWorld* InWorld, FAutomationScreen
 		ContactShadows.Set(0);
 		EyeAdaptationQuality.Set(0);
 		TonemapperGamma.Set(2.2f);
-		//TonemapperSharpen.Set(0);
 	}
 	else if (InOutOptions.bDisableTonemapping)
 	{
 		EyeAdaptationQuality.Set(0);
 		TonemapperGamma.Set(2.2f);
-		//TonemapperSharpen.Set(0);
 	}
 
 	// Forces ScreenPercentage=100
 	{
-		// Completly disable dynamic resolution
+		// Completely disable dynamic resolution
 		{
 			DynamicResTestScreenPercentage.Set(0);
 			DynamicResOperationMode.Set(0);
diff --git a/Engine/Source/Developer/FunctionalTesting/Public/AutomationScreenshotOptions.h b/Engine/Source/Developer/FunctionalTesting/Public/AutomationScreenshotOptions.h
--- a/Engine/Source/Developer/FunctionalTesting/Public/AutomationScreenshotOptions.h
+++ b/Engine/Source/Developer/FunctionalTesting/Public/AutomationScreenshotOptions.h
@@ -64,7 +64,6 @@ public:
 	uint8 MaxBrightness;
 };
 
-
 USTRUCT(BlueprintType)
 struct FAutomationScreenshotOptions
 {
@@ -143,16 +142,17 @@ public:
 	/**
 	 * Disables Anti-Aliasing, Motion Blur, Screen Space Reflections, Eye Adaptation, Tonemapper and Contact
 	 * Shadows, because those features contribute a lot to the noise in the final rendered image.  If you're
-	 * explicitly looking for changes 
+	 * explicitly looking for changes. Unchecking the option will make accessible the Disable Eye Adaptation
+	 * checkbox.
 	 */
-	UPROPERTY(EditAnywhere, BlueprintReadWrite, Category="Screenshot")
+	UPROPERTY(EditAnywhere, BlueprintReadWrite, Category = "Screenshot")
 	bool bDisableNoisyRenderingFeatures;
 
 	/**
 	 * Disables Eye Adaptation and sets Tonemapper to fixed gamma curve. Should generally be on unless
-	 * testing tone mapping or other post-processing results
+	 * testing tone mapping or other post-processing results.
 	 */
-	UPROPERTY(EditAnywhere, BlueprintReadWrite, Category="Screenshot", meta=(DisplayName = "Fixed Exposure"))
+	UPROPERTY(EditAnywhere, BlueprintReadWrite, Category = "Screenshot", meta = (EditCondition = "!bDisableNoisyRenderingFeatures", DisplayName = "Disable Eye Adaptation"))
 	bool bDisableTonemapping;
 
 	/**
