Overhaul stereo variants to support an arbitrary number of variants and correctly implement per-variant frame delays and test preparation, significantly reducing false failures.
Fixes compatibility with TSR sequence screenshot functional test blueprints.

diff --git a/Engine/Source/Developer/FunctionalTesting/Classes/ScreenshotFunctionalTest.h b/Engine/Source/Developer/FunctionalTesting/Classes/ScreenshotFunctionalTest.h
--- a/Engine/Source/Developer/FunctionalTesting/Classes/ScreenshotFunctionalTest.h
+++ b/Engine/Source/Developer/FunctionalTesting/Classes/ScreenshotFunctionalTest.h
@@ -28,20 +28,45 @@ public:
 	UPROPERTY(EditAnywhere, BlueprintReadWrite, Category = "Camera", SimpleDisplay)
 	bool bCameraCutOnScreenshotPrep;
 
+
 protected:
+	virtual bool RunTest(const TArray<FString>& Params = TArray<FString>()) override;
 	virtual void PrepareTest() override;
+	virtual bool IsReady_Implementation() override;
+	virtual void Tick(float DeltaSeconds) override;
+
 	virtual void RequestScreenshot() override;
 	virtual void OnScreenShotCaptured(int32 InSizeX, int32 InSizeY, const TArray<FColor>& InImageData) override;
 	virtual void OnScreenshotTakenAndCompared() override;
-	virtual void StartTest() override;
+	
 
 private:
-	void SetupVariant(FString VariantName, FString SetpCommand, FString RestoreCommand);
 
-	FString CurrentVariantName;
-	FString VariantRestoreCommand;
-	bool bNeedsVariantRestore;
+	enum EVariantType
+	{
+		Baseline = 0,
+		ViewRectOffset,
+		Num
+	};
+
+	struct FVariantInfo
+	{
+		const TCHAR* Name;
+		const TCHAR* SetupCommand;
+		const TCHAR* RestoreCommand;
+	};
+
+	bool SetupNextVariant();
+	void SetupVariant(EVariantType VariantType);
+	
+	TBitArray<> RequestedVariants;
+	const TCHAR* CurrentVariantName;
+	const TCHAR* VariantRestoreCommand;
+
+	// Used for subsequent variants after the first is done
+	void ReprepareTest();
 
-	bool bShouldDoBaselineTest;
-	bool bShouldDoViewRectOffsetVariant;
+	uint32 VariantFrame;
+	float VariantTime;
+	bool bVariantQueued;
 };
diff --git a/Engine/Source/Developer/FunctionalTesting/Private/ScreenshotFunctionalTest.cpp b/Engine/Source/Developer/FunctionalTesting/Private/ScreenshotFunctionalTest.cpp
--- a/Engine/Source/Developer/FunctionalTesting/Private/ScreenshotFunctionalTest.cpp
+++ b/Engine/Source/Developer/FunctionalTesting/Private/ScreenshotFunctionalTest.cpp
@@ -22,9 +22,8 @@
 AScreenshotFunctionalTest::AScreenshotFunctionalTest( const FObjectInitializer& ObjectInitializer )
 	: AScreenshotFunctionalTestBase(ObjectInitializer)
 	, bCameraCutOnScreenshotPrep(true)
-	, bNeedsVariantRestore(false)
-	, bShouldDoBaselineTest(false)
-	, bShouldDoViewRectOffsetVariant(false)
+	, RequestedVariants(false, EVariantType::Num)
+	, bVariantQueued(false)
 {
 }
 
@@ -40,12 +39,32 @@ void AScreenshotFunctionalTest::Serialize(FArchive& Ar)
 	}
 }
 
-void AScreenshotFunctionalTest::PrepareTest()
+bool AScreenshotFunctionalTest::RunTest(const TArray<FString>& Params)
 {
-	// If variants are enabled and lightweight variants are on, skip the baseline test
-	bShouldDoViewRectOffsetVariant = FAutomationTestFramework::NeedPerformStereoTestVariants();
-	bShouldDoBaselineTest = !(bShouldDoViewRectOffsetVariant && FAutomationTestFramework::NeedUseLightweightStereoTestVariants());
+	// Set up which variants are enabled
+	check(!RequestedVariants.Contains(true));
+
+	if (FAutomationTestFramework::NeedPerformStereoTestVariants())
+	{
+		RequestedVariants[EVariantType::Baseline] = !FAutomationTestFramework::NeedUseLightweightStereoTestVariants(); // Skip baseline if lightweight variants are active
+		RequestedVariants[EVariantType::ViewRectOffset] = true;
+
+	}
+	else
+	{
+		RequestedVariants[EVariantType::Baseline] = true;
+		RequestedVariants[EVariantType::ViewRectOffset] = false;
+	}
 
+	// Set up first variant
+	SetupNextVariant();
+
+	// This will call PrepareTest()
+	return Super::RunTest(); 
+}
+
+void AScreenshotFunctionalTest::PrepareTest()
+{
 	// Pre-prep flush to allow rendering to temporary targets and other test resources
 	UAutomationBlueprintFunctionLibrary::FinishLoadingBeforeScreenshot();
 
@@ -136,12 +155,10 @@ void AScreenshotFunctionalTest::OnScreenShotCaptured(int32 InSizeX, int32 InSize
 
 	// If variant in use, pass on the name, then restore settings since capture is done
 	Data.VariantName = CurrentVariantName;
-	CurrentVariantName = "";
-
-	if (bNeedsVariantRestore)
+	if (VariantRestoreCommand)
 	{
-		GEngine->Exec(nullptr, *VariantRestoreCommand);
-		bNeedsVariantRestore = false;
+		GEngine->Exec(nullptr, VariantRestoreCommand);
+		VariantRestoreCommand = nullptr;
 	}
 
 	if (GIsAutomationTesting)
@@ -155,64 +172,121 @@ void AScreenshotFunctionalTest::OnScreenShotCaptured(int32 InSizeX, int32 InSize
 #endif
 }
 
-void AScreenshotFunctionalTest::StartTest()
+void AScreenshotFunctionalTest::OnScreenshotTakenAndCompared()
 {
-	if (bShouldDoBaselineTest)
+	FAutomationTestBase* CurrentTest = FAutomationTestFramework::Get().GetCurrentTest();
+	bool bSkipRemainingVariants = FAutomationTestFramework::Get().NeedUseLightweightStereoTestVariants() && (!CurrentTest || CurrentTest->HasAnyErrors());
+
+	// If we aren't skipping further variants due to a failure, and we still have some remaining, queue the next
+	if (!bSkipRemainingVariants && SetupNextVariant())
 	{
-		bShouldDoBaselineTest = false;
+		ReprepareTest();
 	}
-	else if (bShouldDoViewRectOffsetVariant)
+
+	// Otherwise finish
+	else
 	{
-		bShouldDoViewRectOffsetVariant = false;
-		SetupVariant("ViewRectOffset", "r.Test.ViewRectOffset 5", "r.Test.ViewRectOffset 0");
+		Super::OnScreenshotTakenAndCompared();
 	}
+}
+
+bool AScreenshotFunctionalTest::SetupNextVariant()
+{
+	// Find first remaining requested variant
+	int32 Index = RequestedVariants.Find(true);
 
-	Super::StartTest();
+	if (Index == INDEX_NONE)
+	{
+		return false;
+	}
+	else
+	{
+		// Remove variant from requested variants and set it up
+		RequestedVariants[Index] = false;
+		SetupVariant(static_cast<EVariantType>(Index));
+		return true;
+	}
 }
 
-void AScreenshotFunctionalTest::OnScreenshotTakenAndCompared()
+void AScreenshotFunctionalTest::SetupVariant(EVariantType VariantType)
 {
-	FAutomationTestBase* CurrentTest = FAutomationTestFramework::Get().GetCurrentTest();
-	bool bSkipDueToError = FAutomationTestFramework::Get().NeedUseLightweightStereoTestVariants() && (!CurrentTest || CurrentTest->HasAnyErrors());
+	// Get info for requested variant
+	static const TArray<FVariantInfo> VariantInfoArray =
+	{
+		FVariantInfo{ TEXT(""), nullptr, nullptr }, // Baseline
+		FVariantInfo{ TEXT("ViewRectOffset"), TEXT("r.Test.ViewRectOffset 5"), TEXT("r.Test.ViewRectOffset 0") } // ViewRectOffset
+	};
+
+	const FVariantInfo* VariantInfo = &VariantInfoArray[VariantType];
 
-	// If we still need to perform any variants, loop back here
-	if (bShouldDoViewRectOffsetVariant && !bSkipDueToError)
+	// Set up variant
+	if (VariantInfo->SetupCommand)
 	{
-		// Re-prepare test if necessary
-		if (bCameraCutOnScreenshotPrep)
-		{
-			APlayerController* PlayerController = UGameplayStatics::GetPlayerController(GetWorld(), 0);
+		GEngine->Exec(nullptr, VariantInfo->SetupCommand);
+	}
 
-			if (PlayerController && PlayerController->PlayerCameraManager)
-			{
-				PlayerController->PlayerCameraManager->SetGameCameraCutThisFrame();
-				if (ScreenshotCamera)
-				{
-					ScreenshotCamera->NotifyCameraCut();
-				}
-			}
+	// Save variant name and command needed to restore in OnScreenShotCaptured
+	CurrentVariantName = VariantInfo->Name;
+	if (VariantInfo->RestoreCommand)
+	{
+		VariantRestoreCommand = VariantInfo->RestoreCommand;
+	}
+}
+
+void AScreenshotFunctionalTest::ReprepareTest()
+{
+	// Required to reset TSR sequences
+	OnTestFinished.Broadcast();
+
+	RestoreViewSettings();
+
+	// Rather than re-use RunFrame/RunTime, use a seprarate set for subsequent variants
+	// A screenshot will be requested once the given delay has elapsed
+	VariantFrame = GFrameNumber;
+	VariantTime = (float)GetWorld()->GetTimeSeconds();
+	bVariantQueued = true;
+
+	PrepareTest();
+}
+
+// TSR sequences override this to rely on an internal blueprint instead, but it will still use our Tick() implementation
+// Because we call OnTestFinished.Broadcast() in ReprepareTest(), the TSR BP frame counter is reset properly
+bool AScreenshotFunctionalTest::IsReady_Implementation()
+{
+	if (bVariantQueued)
+	{
+		if ((GetWorld()->GetTimeSeconds() - VariantTime) > ScreenshotOptions.Delay)
+		{
+			return int32(GFrameNumber - VariantFrame) > ScreenshotOptions.FrameDelay;
 		}
 
-		StartTest();
+		return false;
 	}
-
 	else
 	{
-		// This ends the test and reports results
-		Super::OnScreenshotTakenAndCompared();
+		return Super::IsReady_Implementation();
 	}
 }
 
-void AScreenshotFunctionalTest::SetupVariant(FString VariantName, FString SetupCommand, FString RestoreCommand)
+void AScreenshotFunctionalTest::Tick(float DeltaSeconds)
 {
-	// Set up variant
-	GEngine->Exec(nullptr, *SetupCommand);
+	// This section takes over from the main loop to kick off subsequent variants
+	if (bVariantQueued)
+	{
+		if (IsReady())
+		{
+			bVariantQueued = false;
+			StartTest();
+		}
 
-	// Save variant name and command needed to restore in OnScreenShotCaptured
-	CurrentVariantName = VariantName;
-	if (!RestoreCommand.IsEmpty())
+	}
+	else
 	{
-		VariantRestoreCommand = RestoreCommand;
-		bNeedsVariantRestore = true;
+		if (PreparationTimeLimit > 0.f && TotalTime > PreparationTimeLimit)
+		{
+			OnTimeout();
+		}
 	}
-}
+
+	Super::Tick(DeltaSeconds);
+}
\ No newline at end of file
