Fixing crash when trying to restore a CVar to a previous value, when the previous value was empty.

#rnx

diff --git a/Engine/Source/Developer/FunctionalTesting/Private/FunctionalTest.cpp b/Engine/Source/Developer/FunctionalTesting/Private/FunctionalTest.cpp
--- a/Engine/Source/Developer/FunctionalTesting/Private/FunctionalTest.cpp
+++ b/Engine/Source/Developer/FunctionalTesting/Private/FunctionalTest.cpp
@@ -1732,7 +1732,7 @@ void FConsoleVariableBPSetter::Set(const FString& Value)
 			OriginalValue = ConsoleVariable->GetString();
 		}
 
-		ConsoleVariable->AsVariable()->SetWithCurrentPriority(Value.GetCharArray().GetData());
+		ConsoleVariable->AsVariable()->SetWithCurrentPriority(*Value);
 	}
 }
 
@@ -1755,7 +1755,7 @@ void FConsoleVariableBPSetter::Restore()
 		IConsoleVariable* ConsoleVariable = IConsoleManager::Get().FindConsoleVariable(*ConsoleVariableName);
 		if (ensure(ConsoleVariable))
 		{
-			ConsoleVariable->AsVariable()->SetWithCurrentPriority(OriginalValue.GetCharArray().GetData());
+			ConsoleVariable->AsVariable()->SetWithCurrentPriority(*OriginalValue);
 		}
 
 		bModified = false;
