[Backout] - CL30314361
Original CL Desc
-----------------------------------------------------------------
Overhaul stereo variants to support an arbitrary number of variants and correctly implement per-variant frame delays and test preparation, significantly reducing false failures.
Fixes compatibility with TSR sequence screenshot functional test blueprints.

diff --git a/Engine/Source/Developer/FunctionalTesting/Classes/ScreenshotFunctionalTest.h b/Engine/Source/Developer/FunctionalTesting/Classes/ScreenshotFunctionalTest.h
--- a/Engine/Source/Developer/FunctionalTesting/Classes/ScreenshotFunctionalTest.h
+++ b/Engine/Source/Developer/FunctionalTesting/Classes/ScreenshotFunctionalTest.h
@@ -28,45 +28,20 @@ public:
 	UPROPERTY(EditAnywhere, BlueprintReadWrite, Category = "Camera", SimpleDisplay)
 	bool bCameraCutOnScreenshotPrep;
 
-
 protected:
-	virtual bool RunTest(const TArray<FString>& Params = TArray<FString>()) override;
 	virtual void PrepareTest() override;
-	virtual bool IsReady_Implementation() override;
-	virtual void Tick(float DeltaSeconds) override;
-
 	virtual void RequestScreenshot() override;
 	virtual void OnScreenShotCaptured(int32 InSizeX, int32 InSizeY, const TArray<FColor>& InImageData) override;
 	virtual void OnScreenshotTakenAndCompared() override;
-	
+	virtual void StartTest() override;
 
 private:
+	void SetupVariant(FString VariantName, FString SetpCommand, FString RestoreCommand);
 
-	enum EVariantType
-	{
-		Baseline = 0,
-		ViewRectOffset,
-		Num
-	};
-
-	struct FVariantInfo
-	{
-		const TCHAR* Name;
-		const TCHAR* SetupCommand;
-		const TCHAR* RestoreCommand;
-	};
-
-	bool SetupNextVariant();
-	void SetupVariant(EVariantType VariantType);
-	
-	TBitArray<> RequestedVariants;
-	const TCHAR* CurrentVariantName;
-	const TCHAR* VariantRestoreCommand;
-
-	// Used for subsequent variants after the first is done
-	void ReprepareTest();
+	FString CurrentVariantName;
+	FString VariantRestoreCommand;
+	bool bNeedsVariantRestore;
 
-	uint32 VariantFrame;
-	float VariantTime;
-	bool bVariantQueued;
+	bool bShouldDoBaselineTest;
+	bool bShouldDoViewRectOffsetVariant;
 };
diff --git a/Engine/Source/Developer/FunctionalTesting/Private/ScreenshotFunctionalTest.cpp b/Engine/Source/Developer/FunctionalTesting/Private/ScreenshotFunctionalTest.cpp
--- a/Engine/Source/Developer/FunctionalTesting/Private/ScreenshotFunctionalTest.cpp
+++ b/Engine/Source/Developer/FunctionalTesting/Private/ScreenshotFunctionalTest.cpp
@@ -22,8 +22,9 @@
 AScreenshotFunctionalTest::AScreenshotFunctionalTest( const FObjectInitializer& ObjectInitializer )
 	: AScreenshotFunctionalTestBase(ObjectInitializer)
 	, bCameraCutOnScreenshotPrep(true)
-	, RequestedVariants(false, EVariantType::Num)
-	, bVariantQueued(false)
+	, bNeedsVariantRestore(false)
+	, bShouldDoBaselineTest(false)
+	, bShouldDoViewRectOffsetVariant(false)
 {
 }
 
@@ -39,32 +40,12 @@ void AScreenshotFunctionalTest::Serialize(FArchive& Ar)
 	}
 }
 
-bool AScreenshotFunctionalTest::RunTest(const TArray<FString>& Params)
-{
-	// Set up which variants are enabled
-	check(!RequestedVariants.Contains(true));
-
-	if (FAutomationTestFramework::NeedPerformStereoTestVariants())
-	{
-		RequestedVariants[EVariantType::Baseline] = !FAutomationTestFramework::NeedUseLightweightStereoTestVariants(); // Skip baseline if lightweight variants are active
-		RequestedVariants[EVariantType::ViewRectOffset] = true;
-
-	}
-	else
-	{
-		RequestedVariants[EVariantType::Baseline] = true;
-		RequestedVariants[EVariantType::ViewRectOffset] = false;
-	}
-
-	// Set up first variant
-	SetupNextVariant();
-
-	// This will call PrepareTest()
-	return Super::RunTest(); 
-}
-
 void AScreenshotFunctionalTest::PrepareTest()
 {
+	// If variants are enabled and lightweight variants are on, skip the baseline test
+	bShouldDoViewRectOffsetVariant = FAutomationTestFramework::NeedPerformStereoTestVariants();
+	bShouldDoBaselineTest = !(bShouldDoViewRectOffsetVariant && FAutomationTestFramework::NeedUseLightweightStereoTestVariants());
+
 	// Pre-prep flush to allow rendering to temporary targets and other test resources
 	UAutomationBlueprintFunctionLibrary::FinishLoadingBeforeScreenshot();
 
@@ -155,10 +136,12 @@ void AScreenshotFunctionalTest::OnScreenShotCaptured(int32 InSizeX, int32 InSize
 
 	// If variant in use, pass on the name, then restore settings since capture is done
 	Data.VariantName = CurrentVariantName;
-	if (VariantRestoreCommand)
+	CurrentVariantName = "";
+
+	if (bNeedsVariantRestore)
 	{
-		GEngine->Exec(nullptr, VariantRestoreCommand);
-		VariantRestoreCommand = nullptr;
+		GEngine->Exec(nullptr, *VariantRestoreCommand);
+		bNeedsVariantRestore = false;
 	}
 
 	if (GIsAutomationTesting)
@@ -172,121 +155,64 @@ void AScreenshotFunctionalTest::OnScreenShotCaptured(int32 InSizeX, int32 InSize
 #endif
 }
 
-void AScreenshotFunctionalTest::OnScreenshotTakenAndCompared()
-{
-	FAutomationTestBase* CurrentTest = FAutomationTestFramework::Get().GetCurrentTest();
-	bool bSkipRemainingVariants = FAutomationTestFramework::Get().NeedUseLightweightStereoTestVariants() && (!CurrentTest || CurrentTest->HasAnyErrors());
-
-	// If we aren't skipping further variants due to a failure, and we still have some remaining, queue the next
-	if (!bSkipRemainingVariants && SetupNextVariant())
-	{
-		ReprepareTest();
-	}
-
-	// Otherwise finish
-	else
-	{
-		Super::OnScreenshotTakenAndCompared();
-	}
-}
-
-bool AScreenshotFunctionalTest::SetupNextVariant()
+void AScreenshotFunctionalTest::StartTest()
 {
-	// Find first remaining requested variant
-	int32 Index = RequestedVariants.Find(true);
-
-	if (Index == INDEX_NONE)
+	if (bShouldDoBaselineTest)
 	{
-		return false;
+		bShouldDoBaselineTest = false;
 	}
-	else
-	{
-		// Remove variant from requested variants and set it up
-		RequestedVariants[Index] = false;
-		SetupVariant(static_cast<EVariantType>(Index));
-		return true;
-	}
-}
-
-void AScreenshotFunctionalTest::SetupVariant(EVariantType VariantType)
-{
-	// Get info for requested variant
-	static const TArray<FVariantInfo> VariantInfoArray =
-	{
-		FVariantInfo{ TEXT(""), nullptr, nullptr }, // Baseline
-		FVariantInfo{ TEXT("ViewRectOffset"), TEXT("r.Test.ViewRectOffset 5"), TEXT("r.Test.ViewRectOffset 0") } // ViewRectOffset
-	};
-
-	const FVariantInfo* VariantInfo = &VariantInfoArray[VariantType];
-
-	// Set up variant
-	if (VariantInfo->SetupCommand)
+	else if (bShouldDoViewRectOffsetVariant)
 	{
-		GEngine->Exec(nullptr, VariantInfo->SetupCommand);
+		bShouldDoViewRectOffsetVariant = false;
+		SetupVariant("ViewRectOffset", "r.Test.ViewRectOffset 5", "r.Test.ViewRectOffset 0");
 	}
 
-	// Save variant name and command needed to restore in OnScreenShotCaptured
-	CurrentVariantName = VariantInfo->Name;
-	if (VariantInfo->RestoreCommand)
-	{
-		VariantRestoreCommand = VariantInfo->RestoreCommand;
-	}
+	Super::StartTest();
 }
 
-void AScreenshotFunctionalTest::ReprepareTest()
+void AScreenshotFunctionalTest::OnScreenshotTakenAndCompared()
 {
-	// Required to reset TSR sequences
-	OnTestFinished.Broadcast();
-
-	RestoreViewSettings();
-
-	// Rather than re-use RunFrame/RunTime, use a seprarate set for subsequent variants
-	// A screenshot will be requested once the given delay has elapsed
-	VariantFrame = GFrameNumber;
-	VariantTime = (float)GetWorld()->GetTimeSeconds();
-	bVariantQueued = true;
-
-	PrepareTest();
-}
+	FAutomationTestBase* CurrentTest = FAutomationTestFramework::Get().GetCurrentTest();
+	bool bSkipDueToError = FAutomationTestFramework::Get().NeedUseLightweightStereoTestVariants() && (!CurrentTest || CurrentTest->HasAnyErrors());
 
-// TSR sequences override this to rely on an internal blueprint instead, but it will still use our Tick() implementation
-// Because we call OnTestFinished.Broadcast() in ReprepareTest(), the TSR BP frame counter is reset properly
-bool AScreenshotFunctionalTest::IsReady_Implementation()
-{
-	if (bVariantQueued)
+	// If we still need to perform any variants, loop back here
+	if (bShouldDoViewRectOffsetVariant && !bSkipDueToError)
 	{
-		if ((GetWorld()->GetTimeSeconds() - VariantTime) > ScreenshotOptions.Delay)
+		// Re-prepare test if necessary
+		if (bCameraCutOnScreenshotPrep)
 		{
-			return int32(GFrameNumber - VariantFrame) > ScreenshotOptions.FrameDelay;
+			APlayerController* PlayerController = UGameplayStatics::GetPlayerController(GetWorld(), 0);
+
+			if (PlayerController && PlayerController->PlayerCameraManager)
+			{
+				PlayerController->PlayerCameraManager->SetGameCameraCutThisFrame();
+				if (ScreenshotCamera)
+				{
+					ScreenshotCamera->NotifyCameraCut();
+				}
+			}
 		}
 
-		return false;
+		StartTest();
 	}
+
 	else
 	{
-		return Super::IsReady_Implementation();
+		// This ends the test and reports results
+		Super::OnScreenshotTakenAndCompared();
 	}
 }
 
-void AScreenshotFunctionalTest::Tick(float DeltaSeconds)
+void AScreenshotFunctionalTest::SetupVariant(FString VariantName, FString SetupCommand, FString RestoreCommand)
 {
-	// This section takes over from the main loop to kick off subsequent variants
-	if (bVariantQueued)
-	{
-		if (IsReady())
-		{
-			bVariantQueued = false;
-			StartTest();
-		}
+	// Set up variant
+	GEngine->Exec(nullptr, *SetupCommand);
 
-	}
-	else
+	// Save variant name and command needed to restore in OnScreenShotCaptured
+	CurrentVariantName = VariantName;
+	if (!RestoreCommand.IsEmpty())
 	{
-		if (PreparationTimeLimit > 0.f && TotalTime > PreparationTimeLimit)
-		{
-			OnTimeout();
-		}
+		VariantRestoreCommand = RestoreCommand;
+		bNeedsVariantRestore = true;
 	}
-
-	Super::Tick(DeltaSeconds);
-}
\ No newline at end of file
+}
