Fix screenshot comparison slipping to the next test
Add a condition to FinishTest to make sure it gets trigger only if the screenshot was taken and compared.
+ move functional test metadata log to when the test is about to run

#rnx

diff --git a/Engine/Source/Developer/FunctionalTesting/Classes/ScreenshotFunctionalTestBase.h b/Engine/Source/Developer/FunctionalTesting/Classes/ScreenshotFunctionalTestBase.h
--- a/Engine/Source/Developer/FunctionalTesting/Classes/ScreenshotFunctionalTestBase.h
+++ b/Engine/Source/Developer/FunctionalTesting/Classes/ScreenshotFunctionalTestBase.h
@@ -32,6 +32,8 @@ public:
 
 	virtual void Serialize(FArchive& Ar) override;
 
+	virtual void FinishTest(EFunctionalTestResult TestResult, const FString& Message) override;
+
 protected:
 	// Set player view target to screenshot camera and call PrepareForScreenshot
 	virtual void PrepareTest() override;
@@ -62,6 +64,8 @@ protected:
 	// Restore viewport size and original environment settings
 	void RestoreViewSettings();
 
+	virtual void OnTimeout() override;
+
 protected:
 
 	UPROPERTY(EditAnywhere, BlueprintReadOnly, Category = "Screenshot", meta = (MultiLine = "true"))
@@ -82,4 +86,5 @@ protected:
 private:
 	bool bNeedsViewSettingsRestore;
 	bool bNeedsViewportRestore;
+	bool bScreenshotCompleted;
 };
diff --git a/Engine/Source/Developer/FunctionalTesting/Private/ScreenshotFunctionalTestBase.cpp b/Engine/Source/Developer/FunctionalTesting/Private/ScreenshotFunctionalTestBase.cpp
--- a/Engine/Source/Developer/FunctionalTesting/Private/ScreenshotFunctionalTestBase.cpp
+++ b/Engine/Source/Developer/FunctionalTesting/Private/ScreenshotFunctionalTestBase.cpp
@@ -34,6 +34,7 @@ AScreenshotFunctionalTestBase::AScreenshotFunctionalTestBase(const FObjectInitia
 	, ScreenshotOptions(EComparisonTolerance::Low)
 	, bNeedsViewSettingsRestore(false)
 	, bNeedsViewportRestore(false)
+	, bScreenshotCompleted(false)
 {
 	ScreenshotCamera = CreateDefaultSubobject<UCameraComponent>(TEXT("Camera"));
 	ScreenshotCamera->SetupAttachment(RootComponent);
@@ -82,13 +83,29 @@ void AScreenshotFunctionalTestBase::StartTest()
 
 void AScreenshotFunctionalTestBase::OnScreenshotTakenAndCompared()
 {
-	RestoreViewSettings();
-
-	FAutomationTestFramework::Get().OnScreenshotTakenAndCompared.RemoveAll(this);
+	bScreenshotCompleted = true;
 
 	FinishTest(EFunctionalTestResult::Succeeded, TEXT(""));
 }
 
+void AScreenshotFunctionalTestBase::FinishTest(EFunctionalTestResult TestResult, const FString& Message)
+{
+	if (bScreenshotCompleted)
+	{
+		RestoreViewSettings();
+
+		FAutomationTestFramework::Get().OnScreenshotTakenAndCompared.RemoveAll(this);
+
+		Super::FinishTest(TestResult, Message);
+	}
+	else if (TestResult == EFunctionalTestResult::Error
+				|| TestResult == EFunctionalTestResult::Failed
+				|| TestResult == EFunctionalTestResult::Invalid)
+	{
+		AddError(Message);
+	}
+}
+
 void AScreenshotFunctionalTestBase::PrepareForScreenshot()
 {
 	UGameViewportClient* GameViewportClient = AutomationCommon::GetAnyGameViewportClient();
@@ -96,6 +113,8 @@ void AScreenshotFunctionalTestBase::PrepareForScreenshot()
 	check(IsInGameThread());
 	check(!bNeedsViewSettingsRestore && !bNeedsViewportRestore);
 
+	bScreenshotCompleted = false;
+
 #if WITH_AUTOMATION_TESTS
 	bool bApplyScreenshotSettings = true;
 	FSceneViewport* GameViewport = GameViewportClient->GetGameViewport();
@@ -203,6 +222,11 @@ void AScreenshotFunctionalTestBase::OnComparisonComplete(const FAutomationScreen
 {
 	FAutomationTestFramework::Get().OnScreenshotCompared.RemoveAll(this);
 
+	if(!bIsRunning)
+	{
+		return;
+	}
+
 	if (FAutomationTestBase* CurrentTest = FAutomationTestFramework::Get().GetCurrentTest())
 	{
 		CurrentTest->AddEvent(CompareResults.ToAutomationEvent());
@@ -234,6 +258,14 @@ void AScreenshotFunctionalTestBase::RestoreViewSettings()
 	bNeedsViewportRestore = false;
 }
 
+void AScreenshotFunctionalTestBase::OnTimeout()
+{
+	// If the test timed out, make sure the screenshot comparison is cancelled.
+	bScreenshotCompleted = true;
+
+	Super::OnTimeout();
+}
+
 #if WITH_EDITOR
 
 bool AScreenshotFunctionalTestBase::CanEditChange(const FProperty* InProperty) const
