Fix screenshot timing out if FinishTest is triggered when not ready.

#rnx

diff --git a/Engine/Source/Developer/FunctionalTesting/Private/ScreenshotFunctionalTestBase.cpp b/Engine/Source/Developer/FunctionalTesting/Private/ScreenshotFunctionalTestBase.cpp
--- a/Engine/Source/Developer/FunctionalTesting/Private/ScreenshotFunctionalTestBase.cpp
+++ b/Engine/Source/Developer/FunctionalTesting/Private/ScreenshotFunctionalTestBase.cpp
@@ -90,7 +90,7 @@ void AScreenshotFunctionalTestBase::OnScreenshotTakenAndCompared()
 
 void AScreenshotFunctionalTestBase::FinishTest(EFunctionalTestResult TestResult, const FString& Message)
 {
-	if (bScreenshotCompleted)
+	if (!IsReady() ||  bScreenshotCompleted)
 	{
 		RestoreViewSettings();
 
