FuncTest: Fix up visilbity issue.

diff --git a/Engine/Source/Developer/FunctionalTesting/Classes/FunctionalTestBase.h b/Engine/Source/Developer/FunctionalTesting/Classes/FunctionalTestBase.h
--- a/Engine/Source/Developer/FunctionalTesting/Classes/FunctionalTestBase.h
+++ b/Engine/Source/Developer/FunctionalTesting/Classes/FunctionalTestBase.h
@@ -102,6 +102,6 @@ protected:
 
 private:
 
-	static bool bIsFunctionalTestRunning;
+	FUNCTIONALTESTING_API static bool bIsFunctionalTestRunning;
 	static FString ActiveTestName;
 };
