Skip actor iteration on world CDO.

diff --git a/Engine/Source/Developer/FunctionalTesting/Private/FunctionalTestingModule.cpp b/Engine/Source/Developer/FunctionalTesting/Private/FunctionalTestingModule.cpp
--- a/Engine/Source/Developer/FunctionalTesting/Private/FunctionalTestingModule.cpp
+++ b/Engine/Source/Developer/FunctionalTesting/Private/FunctionalTestingModule.cpp
@@ -64,6 +64,11 @@ void FFunctionalTestingModule::ShutdownModule()
 void FFunctionalTestingModule::OnGetAssetTagsForWorld(const UWorld* World, FAssetRegistryTagsContext Context)
 {
 #if WITH_EDITOR
+	if (!World || World->HasAllFlags(RF_ClassDefaultObject))
+	{
+		return;
+	}
+
 	TArray<FString> TestNamesRuntime;
 	TArray<FString> TestNamesEditor;
 	for (TActorIterator<AFunctionalTest> ActorItr(const_cast<UWorld*>(World), AFunctionalTest::StaticClass(), EActorIteratorFlags::AllActors); ActorItr; ++ActorItr)
