Capture LogRHI errors and warnings when RHIValidation is enabled even when BP functionalt test is not actively running

#rnx

diff --git a/Engine/Source/Developer/FunctionalTesting/Classes/FunctionalTestBase.h b/Engine/Source/Developer/FunctionalTesting/Classes/FunctionalTestBase.h
--- a/Engine/Source/Developer/FunctionalTesting/Classes/FunctionalTestBase.h
+++ b/Engine/Source/Developer/FunctionalTesting/Classes/FunctionalTestBase.h
@@ -5,6 +5,10 @@
 #include "CoreMinimal.h"
 #include "Misc/AutomationTest.h"
 
+#if ENABLE_RHI_VALIDATION
+#include "RHIValidationCommon.h"
+#endif
+
 /**
  * Base class for Functional test cases.  
  */
@@ -20,9 +24,34 @@ public:
 	 */
 	virtual bool SuppressLogs()
 	{
+#if ENABLE_RHI_VALIDATION
+		if (GRHIValidationEnabled && !bSuppressLogs)
+		{
+			// While RHI Validation is enabled, do not suppress log unless explicitly enabled
+			return false;
+		}
+#endif
 		return bSuppressLogs || !IsFunctionalTestRunning();
 	}
 
+	/**
+	 * Should the log category be captured and surfaced as part of the test.
+	 *
+	 * @return true to allow a log category through.
+	 */
+	virtual bool ShouldCaptureLogCategory(const class FName& Category) const
+	{
+#if ENABLE_RHI_VALIDATION
+		if (GRHIValidationEnabled && !IsFunctionalTestRunning())
+		{
+			// If capturing log while functional test is not running, only filter in LogRHI channel. 
+			return Category == FName(TEXT("LogRHI"));
+		}
+#endif
+		return true; 
+	}
+
+
 	/**
 	 * Specify how log errors & warnings should be handled during tests. If values are not set then the project
 	 * defaults will be used.
diff --git a/Engine/Source/Runtime/AutomationWorker/Private/AutomationWorkerModule.cpp b/Engine/Source/Runtime/AutomationWorker/Private/AutomationWorkerModule.cpp
--- a/Engine/Source/Runtime/AutomationWorker/Private/AutomationWorkerModule.cpp
+++ b/Engine/Source/Runtime/AutomationWorker/Private/AutomationWorkerModule.cpp
@@ -210,7 +210,7 @@ void FAutomationWorkerModule::ReportTestComplete()
 			Message->Duration = ExecutionInfo.Duration;
 			// Prune log entries if it is a success to reduce foot print.
 			Message->Entries = bPruneLogsOnSuccess && bSuccess ?
-				ExecutionInfo.GetEntries().FilterByPredicate([](const FAutomationExecutionEntry& Entry) { return Entry.Event.Context != "log"; }) : ExecutionInfo.GetEntries();
+				ExecutionInfo.GetEntries().FilterByPredicate([](const FAutomationExecutionEntry& Entry) { return Entry.Event.Context != TEXT("log"); }) : ExecutionInfo.GetEntries();
 			Message->WarningTotal = ExecutionInfo.GetWarningTotal();
 			Message->ErrorTotal = ExecutionInfo.GetErrorTotal();
 
