Reverting changelist 34647546 (linker issue)

#rnx

diff --git a/Engine/Source/Developer/FunctionalTesting/Classes/FunctionalTestBase.h b/Engine/Source/Developer/FunctionalTesting/Classes/FunctionalTestBase.h
--- a/Engine/Source/Developer/FunctionalTesting/Classes/FunctionalTestBase.h
+++ b/Engine/Source/Developer/FunctionalTesting/Classes/FunctionalTestBase.h
@@ -5,10 +5,6 @@
 #include "CoreMinimal.h"
 #include "Misc/AutomationTest.h"
 
-#if ENABLE_RHI_VALIDATION
-#include "RHIValidationCommon.h"
-#endif
-
 /**
  * Base class for Functional test cases.  
  */
@@ -24,34 +20,9 @@ public:
 	 */
 	virtual bool SuppressLogs()
 	{
-#if ENABLE_RHI_VALIDATION
-		if (GRHIValidationEnabled && !bSuppressLogs)
-		{
-			// While RHI Validation is enabled, do not suppress log unless explicitly enabled
-			return false;
-		}
-#endif
 		return bSuppressLogs || !IsFunctionalTestRunning();
 	}
 
-	/**
-	 * Should the log category be captured and surfaced as part of the test.
-	 *
-	 * @return true to allow a log category through.
-	 */
-	virtual bool ShouldCaptureLogCategory(const class FName& Category) const
-	{
-#if ENABLE_RHI_VALIDATION
-		if (GRHIValidationEnabled && !IsFunctionalTestRunning())
-		{
-			// If capturing log while functional test is not running, only filter in LogRHI channel. 
-			return Category == FName(TEXT("LogRHI"));
-		}
-#endif
-		return true; 
-	}
-
-
 	/**
 	 * Specify how log errors & warnings should be handled during tests. If values are not set then the project
 	 * defaults will be used.
