Capture LogRHI errors and warnings when RHIValidation is enabled even when BP functionalt test is not actively running (second try)

#rnx

diff --git a/Engine/Source/Developer/FunctionalTesting/Classes/FunctionalTestBase.h b/Engine/Source/Developer/FunctionalTesting/Classes/FunctionalTestBase.h
--- a/Engine/Source/Developer/FunctionalTesting/Classes/FunctionalTestBase.h
+++ b/Engine/Source/Developer/FunctionalTesting/Classes/FunctionalTestBase.h
@@ -6,7 +6,7 @@
 #include "Misc/AutomationTest.h"
 
 /**
- * Base class for Functional test cases.  
+ * Base class for Functional test cases.
  */
 class FFunctionalTestBase : public FAutomationTestBase
 {
@@ -18,10 +18,15 @@ public:
 	 *
 	 * @return true to suppress logs
 	 */
-	virtual bool SuppressLogs()
-	{
-		return bSuppressLogs || !IsFunctionalTestRunning();
-	}
+	FUNCTIONALTESTING_API virtual bool SuppressLogs();
+
+	/**
+	 * Should the log category be captured and surfaced as part of the test.
+	 *
+	 * @return true to allow a log category through.
+	 */
+	FUNCTIONALTESTING_API virtual bool ShouldCaptureLogCategory(const class FName& Category) const;
+
 
 	/**
 	 * Specify how log errors & warnings should be handled during tests. If values are not set then the project
@@ -89,7 +94,7 @@ public:
 	/**
 	 * Returns true if a functional test is running (does not include map setup)
 	 */
-	FUNCTIONALTESTING_API static bool IsFunctionalTestRunning()	{ return bIsFunctionalTestRunning;	}
+	FUNCTIONALTESTING_API static bool IsFunctionalTestRunning() { return bIsFunctionalTestRunning; }
 
 protected:
 
diff --git a/Engine/Source/Developer/FunctionalTesting/Private/FunctionalTestBase.cpp b/Engine/Source/Developer/FunctionalTesting/Private/FunctionalTestBase.cpp
--- a/Engine/Source/Developer/FunctionalTesting/Private/FunctionalTestBase.cpp
+++ b/Engine/Source/Developer/FunctionalTesting/Private/FunctionalTestBase.cpp
@@ -4,6 +4,9 @@
 #include "FunctionalTestingModule.h"
 #include "AutomationControllerSettings.h"
 
+#if ENABLE_RHI_VALIDATION
+#include "RHIValidationCommon.h"
+#endif
 
 // statics
 bool FFunctionalTestBase::bIsFunctionalTestRunning;
@@ -21,6 +24,30 @@ FFunctionalTestBase::FFunctionalTestBase(const FString& InName, const bool bInCo
 	bElevateLogWarningsToErrors = true;
 }
 
+bool FFunctionalTestBase::SuppressLogs()
+{
+#if ENABLE_RHI_VALIDATION
+	if (GRHIValidationEnabled && !bSuppressLogs)
+	{
+		// While RHI Validation is enabled, do not suppress log unless explicitly enabled
+		return false;
+	}
+#endif
+	return bSuppressLogs || !IsFunctionalTestRunning();
+}
+
+bool FFunctionalTestBase::ShouldCaptureLogCategory(const class FName& Category) const
+{
+#if ENABLE_RHI_VALIDATION
+	if (GRHIValidationEnabled && !IsFunctionalTestRunning())
+	{
+		// If capturing log while functional test is not running, only filter in LogRHI channel. 
+		return Category == FName(TEXT("LogRHI"));
+	}
+#endif
+	return true;
+}
+
 void FFunctionalTestBase::SetLogErrorAndWarningHandlingToDefault()
 {
 	// Set to project defaults
