Avoid sending screenshot data if functional test is not running

#rnx

diff --git a/Engine/Source/Developer/FunctionalTesting/Private/ScreenshotFunctionalTest.cpp b/Engine/Source/Developer/FunctionalTesting/Private/ScreenshotFunctionalTest.cpp
--- a/Engine/Source/Developer/FunctionalTesting/Private/ScreenshotFunctionalTest.cpp
+++ b/Engine/Source/Developer/FunctionalTesting/Private/ScreenshotFunctionalTest.cpp
@@ -131,6 +131,12 @@ void AScreenshotFunctionalTest::OnScreenShotCaptured(int32 InSizeX, int32 InSize
 	GameViewportClient->OnScreenshotCaptured().RemoveAll(this);
 
 #if WITH_AUTOMATION_TESTS
+	if (!IsRunning())
+	{
+		// Don't send the data if the test is no longer running.
+		return;
+	}
+
 	const FString Context = AutomationCommon::GetWorldContext(GetWorld());
 
 	TArray<uint8> CapturedFrameTrace = AutomationCommon::CaptureFrameTrace(Context, TestLabel);
diff --git a/Engine/Source/Developer/FunctionalTesting/Private/ScreenshotFunctionalTestBase.cpp b/Engine/Source/Developer/FunctionalTesting/Private/ScreenshotFunctionalTestBase.cpp
--- a/Engine/Source/Developer/FunctionalTesting/Private/ScreenshotFunctionalTestBase.cpp
+++ b/Engine/Source/Developer/FunctionalTesting/Private/ScreenshotFunctionalTestBase.cpp
@@ -166,6 +166,12 @@ void AScreenshotFunctionalTestBase::OnScreenShotCaptured(int32 InSizeX, int32 In
 	GameViewportClient->OnScreenshotCaptured().RemoveAll(this);
 
 #if WITH_AUTOMATION_TESTS
+	if (!IsRunning())
+	{
+		// Don't send the data if the test is no longer running.
+		return;
+	}
+
 	const FString Context = AutomationCommon::GetWorldContext(GetWorld());
 
 	TArray<uint8> CapturedFrameTrace = AutomationCommon::CaptureFrameTrace(Context, TestLabel);
