FunctionalTest tagging and filtering
Exposes a "Test Tags" field to the Editor UI for Actor objects implementing FunctionalTest, such as BlueprintFunctionalTest and ScreenshotFunctionalTest. This persists user-defined tag data via the asset cache, and supports hot-editing.

diff --git a/Engine/Source/Developer/FunctionalTesting/Classes/FunctionalTest.h b/Engine/Source/Developer/FunctionalTesting/Classes/FunctionalTest.h
--- a/Engine/Source/Developer/FunctionalTesting/Classes/FunctionalTest.h
+++ b/Engine/Source/Developer/FunctionalTesting/Classes/FunctionalTest.h
@@ -208,6 +208,16 @@ enum class EFunctionalTestResult : uint8
 /* Return a readable string of the provided EFunctionalTestResult enum */
 FString FUNCTIONALTESTING_API LexToString(const EFunctionalTestResult TestResult);
 
+/** Return a dot-separated path prefix string representing the map that contains a test */
+FString FUNCTIONALTESTING_API MapPackageToAutomationPath(const FString& MapPackageName);
+
+/** Registration information for an individual test */
+struct FFunctionalTestInfo
+{
+	FString BeautifiedName;
+	FString TestCommand;
+	FString TestTags;
+};
 
 UENUM(BlueprintType)
 enum class EFunctionalTestLogHandling : uint8
@@ -249,6 +259,13 @@ public:
 	UPROPERTY(EditAnywhere, BlueprintReadOnly, Category = "Functional Testing", meta = (MultiLine = "true"))
 	FString Description;
 
+	/**
+	 * Tags describing this test separated by square brackets, such as '[dog]' or '[cat]' or '[Graphics][prio0][unstable]'.
+	 * Tags can be used to run subsets of tests, or to categorize data in test reports.
+	 */
+	UPROPERTY(EditAnywhere, BlueprintReadOnly, Category = "Functional Testing")
+	FString TestTags;
+
 private:
 	UPROPERTY()
 	TObjectPtr<UBillboardComponent> SpriteComponent;
diff --git a/Engine/Source/Developer/FunctionalTesting/Private/FunctionalTest.cpp b/Engine/Source/Developer/FunctionalTesting/Private/FunctionalTest.cpp
--- a/Engine/Source/Developer/FunctionalTesting/Private/FunctionalTest.cpp
+++ b/Engine/Source/Developer/FunctionalTesting/Private/FunctionalTest.cpp
@@ -6,6 +6,7 @@
 #include "Misc/Paths.h"
 #include "Engine/GameViewportClient.h"
 #include "Engine/LatentActionManager.h"
+#include "Engine/Level.h"
 #include "Components/BillboardComponent.h"
 #include "HAL/FileManager.h"
 #include "Misc/FileHelper.h"
@@ -117,6 +118,21 @@ FString LexToString(const EFunctionalTestResult TestResult)
 }
 
 
+FString MapPackageToAutomationPath(const FString& MapPackageName)
+{
+	FString PartialSuiteName;
+	if (MapPackageName.StartsWith(TEXT("/Game/")))
+	{
+		PartialSuiteName = MapPackageName.RightChop(6); // Remove "/Game/" from the name, it is not descriptive
+	}
+	else
+	{
+		PartialSuiteName = MapPackageName.RightChop(1); // Remove leading slash
+	}
+	return PartialSuiteName.Replace(TEXT("/"), TEXT(".")); // use dot syntax
+}
+
+
 AFunctionalTest::AFunctionalTest( const FObjectInitializer& ObjectInitializer )
 	: Super(ObjectInitializer)
 	, TestLabel(GetName())
@@ -262,6 +278,7 @@ bool AFunctionalTest::RunTest(const TArray<FString>& Params)
 		{
 			AddInfo(FString::Printf(TEXT("[Owner] %s"), *Author));
 			AddInfo(FString::Printf(TEXT("[Description] %s"), *Description));
+			AddInfo(FString::Printf(TEXT("[TestTags] %s"), *TestTags));
 		}
 	}
 
@@ -569,9 +586,10 @@ void AFunctionalTest::RegisterAutoDestroyActor(AActor* ActorToAutoDestroy)
 
 void AFunctionalTest::PostEditChangeProperty( struct FPropertyChangedEvent& PropertyChangedEvent)
 {
-	static const FName NAME_FunctionalTesting = FName(TEXT("FunctionalTesting"));
+	static const FName NAME_FunctionalTesting = FName(TEXT("Functional Testing"));
 	static const FName NAME_TimeLimit = FName(TEXT("TimeLimit"));
 	static const FName NAME_TimesUpResult = FName(TEXT("TimesUpResult"));
+	static const FName NAME_TestTags = FName(TEXT("TestTags"));
 
 	Super::PostEditChangeProperty(PropertyChangedEvent);
 
@@ -594,6 +612,18 @@ void AFunctionalTest::PostEditChangeProperty( struct FPropertyChangedEvent& Prop
 					TimesUpResult = EFunctionalTestResult::Failed;
 				}
 			}
+			else if (PropertyChangedEvent.Property->GetFName() == NAME_TestTags)
+			{
+				// re-register existing tags for the loaded test
+				FString CurrentMapPackageName = this->GetLevel()->GetPackage()->GetName();
+				FString PartialSuiteName = MapPackageToAutomationPath(CurrentMapPackageName);
+				FString FullBeautifiedName(PartialSuiteName + TEXT(".") + *GetActorLabel());
+				FString FullTestName(TEXT("Project.Functional Tests." + FullBeautifiedName));
+
+				FAutomationTestFramework& TestFramework = FAutomationTestFramework::Get();
+				TestFramework.UnregisterAutomationTestTags(FullTestName);
+				TestFramework.RegisterAutomationTestTags(FullTestName, TestTags);
+			}
 		}
 	}
 }
@@ -611,7 +641,7 @@ void AFunctionalTest::GetAssetRegistryTags(FAssetRegistryTagsContext Context) co
 
 	if (IsPackageExternal() && IsEnabled())
 	{
-		const FString TestActor = GetActorLabel() + TEXT("|") + GetName();
+		const FString TestActor = GetActorLabel() + TEXT("|") + GetName() + TEXT("|") + TestTags;
 		const TCHAR* TestCategory = IsEditorOnlyObject(this) ? TEXT("TestNameEditor") : TEXT("TestName");
 		Context.AddTag(UObject::FAssetRegistryTag(TestCategory, TestActor, UObject::FAssetRegistryTag::TT_Hidden));
 	}
diff --git a/Engine/Source/Developer/FunctionalTesting/Private/FunctionalTestingModule.cpp b/Engine/Source/Developer/FunctionalTesting/Private/FunctionalTestingModule.cpp
--- a/Engine/Source/Developer/FunctionalTesting/Private/FunctionalTestingModule.cpp
+++ b/Engine/Source/Developer/FunctionalTesting/Private/FunctionalTestingModule.cpp
@@ -36,6 +36,7 @@ public:
 	virtual class UFunctionalTestingManager* GetCurrentManager();
 	virtual void SetLooping(const bool bLoop) override;
 	virtual void GetMapTests(bool bEditorOnlyTests, TArray<FString>& OutBeautifiedNames, TArray<FString>& OutTestCommands, TArray<FString>& OutTestMapAssets) const override;
+	virtual void GetMapTests(bool bEditorOnlyTests, TArray<FFunctionalTestInfo>& OutTestInfo, TArray<FString>& OutTestMapAssets) const override;
 
 private:
 	UWorld* GetTestWorld();
@@ -81,7 +82,7 @@ void FFunctionalTestingModule::OnGetAssetTagsForWorld(const UWorld* World, FAsse
 			if (FunctionalTest->IsEnabled())
 			{
 				TArray<FString>& TestNames = IsEditorOnlyObject(FunctionalTest) ? TestNamesEditor : TestNamesRuntime;
-				TestNames.Add(FString::Printf(TEXT("%s|%s;"), *FunctionalTest->GetActorLabel(), *FunctionalTest->GetName()));
+				TestNames.Add(FString::Printf(TEXT("%s|%s|%s;"), *FunctionalTest->GetActorLabel(), *FunctionalTest->GetName(), *FunctionalTest->TestTags));
 			}
 		}
 	}
@@ -101,6 +102,17 @@ void FFunctionalTestingModule::OnGetAssetTagsForWorld(const UWorld* World, FAsse
 }
 
 void FFunctionalTestingModule::GetMapTests(bool bEditorOnlyTests, TArray<FString>& OutBeautifiedNames, TArray<FString>& OutTestCommands, TArray<FString>& OutTestMapAssets) const
+{
+	TArray<FFunctionalTestInfo> TestInfo;
+	GetMapTests(bEditorOnlyTests, TestInfo, OutTestMapAssets);
+	for (FFunctionalTestInfo Info : TestInfo)
+	{
+		OutBeautifiedNames.Add(Info.BeautifiedName);
+		OutTestCommands.Add(Info.TestCommand);
+	}
+}
+
+void FFunctionalTestingModule::GetMapTests(bool bEditorOnlyTests, TArray<FFunctionalTestInfo>& OutTestInfo, TArray<FString>& OutTestMapAssets) const
 {
 	IAssetRegistry& AssetRegistry = FModuleManager::Get().LoadModuleChecked<FAssetRegistryModule>("AssetRegistry").Get();
 
@@ -133,18 +145,14 @@ void FFunctionalTestingModule::GetMapTests(bool bEditorOnlyTests, TArray<FString
 				FString MapAssetPath = MapAsset.GetObjectPathString();
 				FString MapPackageName = MapAsset.PackageName.ToString();
 				if (!IsDeveloperDirectoryIncluded && MapPackageName.Find(TEXT("/Game/Developers")) == 0) continue;
-				FString PartialSuiteName = MapPackageName.RightChop(1).Replace(TEXT("/"), TEXT(".")); // use dot syntax
-				if (MapPackageName.StartsWith(TEXT("/Game/")))
-				{
-					PartialSuiteName = PartialSuiteName.RightChop(5); // Remove "/Game/" from the name
-				}
+				FString PartialSuiteName = MapPackageToAutomationPath(MapPackageName);
 
-				FString AllTestNames;
+				TArray<FString> MapTests;
 				FAssetDataTagMapSharedView::FFindTagResult MapTestNames = MapAsset.TagsAndValues.FindTag(bEditorOnlyTests ? TEXT("TestNamesEditor") : TEXT("TestNames"));
 
 				if (MapTestNames.IsSet())
 				{
-					AllTestNames = MapTestNames.GetValue();
+					MapTests.Add(MapTestNames.GetValue());
 				}
 
 #if WITH_EDITOR
@@ -170,38 +178,41 @@ void FFunctionalTestingModule::GetMapTests(bool bEditorOnlyTests, TArray<FString
 
 						if (ActorTestName.IsSet())
 						{
-							if (!AllTestNames.IsEmpty())
-							{
-								AllTestNames += TEXT(";");
-							}
-
-							AllTestNames += ActorTestName.GetValue();
+							MapTests.Add(ActorTestName.GetValue());
 						}
 					}
 				}
 #endif
-				if (!AllTestNames.IsEmpty())
+				if (!MapTests.IsEmpty())
 				{
-					TArray<FString> MapTests;
-					AllTestNames.ParseIntoArray(MapTests, TEXT(";"), true);
-
 					for (const FString& MapTest : MapTests)
 					{
+						FString Remainder;
 						FString BeautifulTestName;
 						FString RealTestName;
+						FString TestTags;
 
-						if (MapTest.Split(TEXT("|"), &BeautifulTestName, &RealTestName))
+						if (MapTest.Split(TEXT("|"), &BeautifulTestName, &Remainder))
 						{
-							OutBeautifiedNames.Add(PartialSuiteName + TEXT(".") + *BeautifulTestName);
-							OutTestCommands.Add(MapAssetPath + TEXT(";") + MapPackageName + TEXT(";") + *RealTestName);
+							if (!Remainder.Split(TEXT("|"), &RealTestName, &TestTags))
+							{
+								// split fails when loading an old format test-asset without tags
+								RealTestName = Remainder; // tags remain empty
+							}
+
+							FString FullBeautifiedName(PartialSuiteName + TEXT(".") + *BeautifulTestName);
+							FString TestCommand(MapAssetPath + TEXT(";") + MapPackageName + TEXT(";") + *RealTestName);
+							OutTestInfo.Add(FFunctionalTestInfo(FullBeautifiedName, TestCommand, TestTags));
 							OutTestMapAssets.AddUnique(MapAssetPath);
 						}
 					}
 				}
-				else if (!bEditorOnlyTests && MapAsset.AssetName.ToString().Find(TEXT("FTEST_")) == 0)
+				else if (!bEditorOnlyTests && MapAsset.AssetName.ToString().StartsWith(TEXT("FTEST_")))
 				{
-					OutBeautifiedNames.Add(MapAsset.AssetName.ToString());
-					OutTestCommands.Add(MapAssetPath + TEXT(";") + MapPackageName);
+					// add legacy functional test maps which autoplay
+					FString FullBeautifiedName(MapAsset.AssetName.ToString());
+					FString TestCommand(MapAssetPath + TEXT(";") + MapPackageName);
+					OutTestInfo.Add(FFunctionalTestInfo(FullBeautifiedName, TestCommand, "")); // Tags not supported
 					OutTestMapAssets.AddUnique(MapAssetPath);
 				}
 			}
diff --git a/Engine/Source/Developer/FunctionalTesting/Private/Tests/ClientFuncTestPerforming.cpp b/Engine/Source/Developer/FunctionalTesting/Private/Tests/ClientFuncTestPerforming.cpp
--- a/Engine/Source/Developer/FunctionalTesting/Private/Tests/ClientFuncTestPerforming.cpp
+++ b/Engine/Source/Developer/FunctionalTesting/Private/Tests/ClientFuncTestPerforming.cpp
@@ -79,8 +79,20 @@ public:
 	virtual void GetTests(TArray<FString>& OutBeautifiedNames, TArray<FString>& OutTestCommands) const override
 	{
 		bool bEditorOnlyTests = !(GetTestFlags() & EAutomationTestFlags::ClientContext);
-		TArray<FString> MapAssets;
-		IFunctionalTestingModule::Get().GetMapTests(bEditorOnlyTests, OutBeautifiedNames, OutTestCommands, MapAssets);
+		TArray<FString> MapAssetsUnused;
+		TArray<FFunctionalTestInfo> AllTestInfo;
+		IFunctionalTestingModule::Get().GetMapTests(bEditorOnlyTests, AllTestInfo, MapAssetsUnused);
+		FAutomationTestFramework& Framework = FAutomationTestFramework::Get();
+		for (FFunctionalTestInfo TestInfo : AllTestInfo)
+		{
+			OutBeautifiedNames.Add(TestInfo.BeautifiedName);
+			OutTestCommands.Add(TestInfo.TestCommand);
+			if (!TestInfo.TestTags.IsEmpty())
+			{
+				// Register new tags
+				Framework.RegisterComplexAutomationTestTags(this, TestInfo.BeautifiedName, TestInfo.TestTags);
+			}
+		}
 	}
 
 	/**
diff --git a/Engine/Source/Developer/FunctionalTesting/Public/FunctionalTestingModule.h b/Engine/Source/Developer/FunctionalTesting/Public/FunctionalTestingModule.h
--- a/Engine/Source/Developer/FunctionalTesting/Public/FunctionalTestingModule.h
+++ b/Engine/Source/Developer/FunctionalTesting/Public/FunctionalTestingModule.h
@@ -5,6 +5,7 @@
 #include "CoreMinimal.h"
 #include "Modules/ModuleInterface.h"
 #include "Modules/ModuleManager.h"
+#include "FunctionalTest.h"
 
 FUNCTIONALTESTING_API DECLARE_LOG_CATEGORY_EXTERN(LogFunctionalTest, Log, All);
 
@@ -39,8 +40,17 @@ public:
 	virtual void SetLooping(const bool bLoop) = 0;
 
 	/** Gets a list of maps/tests in the current project */
+	UE_DEPRECATED(5.5, "GetMapTests() with four parameters has been deprecated.")
 	virtual void GetMapTests(bool bEditorOnlyTests, TArray<FString>& OutBeautifiedNames, TArray<FString>& OutTestCommands, TArray<FString>& OutTestMapAssets) const = 0;
 
+	/**
+	 * Gets a list of tests and maps in the current project
+	 * @param bEditorOnlyTests	Only collect tests which run in the Editor
+	 * @param OutTestInfo		Updated to contain execution information for each test
+	 * @param OutTestMapAssets	Updated with the unique *.umap assets which contain tests
+	 */
+	virtual void GetMapTests(bool bEditorOnlyTests, TArray<FFunctionalTestInfo>& OutTestInfo, TArray<FString>& OutTestMapAssets) const = 0;
+
 	/** Gets the debugger singleton or returns NULL */
 	static IFunctionalTestingModule& Get()
 	{
diff --git a/Engine/Source/Runtime/Core/Public/Misc/AutomationTest.h b/Engine/Source/Runtime/Core/Public/Misc/AutomationTest.h
--- a/Engine/Source/Runtime/Core/Public/Misc/AutomationTest.h
+++ b/Engine/Source/Runtime/Core/Public/Misc/AutomationTest.h
@@ -1514,10 +1514,18 @@ public:
 		}
 	}
 
+	/**
+	 * No-op default ctor performs no registration
+	 */
+	FAutomationTestTags() {}
+
 	~FAutomationTestTags()
 	{
-		//ignore failure to unregister due to already being unregistered
-		FAutomationTestFramework::Get().UnregisterAutomationTestTags(TestFullName);
+		if (!TestFullName.IsEmpty())
+		{
+			//ignore failure to unregister due to already being unregistered
+			FAutomationTestFramework::Get().UnregisterAutomationTestTags(TestFullName);
+		}
 	}
 private:
 	FString TestFullName;
