bugfix for FunctionalTests failing to register

diff --git a/Engine/Source/Developer/FunctionalTesting/Private/FunctionalTestingModule.cpp b/Engine/Source/Developer/FunctionalTesting/Private/FunctionalTestingModule.cpp
--- a/Engine/Source/Developer/FunctionalTesting/Private/FunctionalTestingModule.cpp
+++ b/Engine/Source/Developer/FunctionalTesting/Private/FunctionalTestingModule.cpp
@@ -148,11 +148,11 @@ void FFunctionalTestingModule::GetMapTests(bool bEditorOnlyTests, TArray<FFuncti
 				FString PartialSuiteName = MapPackageToAutomationPath(MapPackageName);
 
 				TArray<FString> MapTests;
-				FAssetDataTagMapSharedView::FFindTagResult MapTestNames = MapAsset.TagsAndValues.FindTag(bEditorOnlyTests ? TEXT("TestNamesEditor") : TEXT("TestNames"));
+				FAssetDataTagMapSharedView::FFindTagResult MapAssetNamesResult = MapAsset.TagsAndValues.FindTag(bEditorOnlyTests ? TEXT("TestNamesEditor") : TEXT("TestNames"));
 
-				if (MapTestNames.IsSet())
+				if (MapAssetNamesResult.IsSet())
 				{
-					MapTests.Add(MapTestNames.GetValue());
+					MapAssetNamesResult.GetValue().ParseIntoArray(MapTests, TEXT(";"), true);
 				}
 
 #if WITH_EDITOR
