Functional Tests: Add Test Finished BP event

#rnx

diff --git a/Engine/Source/Developer/FunctionalTesting/Classes/FunctionalTest.h b/Engine/Source/Developer/FunctionalTesting/Classes/FunctionalTest.h
--- a/Engine/Source/Developer/FunctionalTesting/Classes/FunctionalTest.h
+++ b/Engine/Source/Developer/FunctionalTesting/Classes/FunctionalTest.h
@@ -761,6 +761,13 @@ protected:
 	UFUNCTION(BlueprintImplementableEvent, meta=( DisplayName="Start Test" ))
 	void ReceiveStartTest();
 
+	/**
+	 * Called during FinishTest().  Allows for clean-up on the blueprint side, so that the user can bring the test
+	 * back to its default state, making it ready for the next one.
+	 */
+	UFUNCTION(BlueprintImplementableEvent, meta=( DisplayName="Test Finished" ))
+	void ReceiveTestFinished();
+
 	/**
 	 * IsReady() is called once per frame after a test is run, until it returns true.  You should use this function to
 	 * delay Start being called on the test until preconditions are met.
diff --git a/Engine/Source/Developer/FunctionalTesting/Private/FunctionalTest.cpp b/Engine/Source/Developer/FunctionalTesting/Private/FunctionalTest.cpp
--- a/Engine/Source/Developer/FunctionalTesting/Private/FunctionalTest.cpp
+++ b/Engine/Source/Developer/FunctionalTesting/Private/FunctionalTest.cpp
@@ -434,6 +434,7 @@ void AFunctionalTest::FinishTest(EFunctionalTestResult TestResult, const FString
 	bIsRunning = false;
 	SetActorTickEnabled(false);
 
+	ReceiveTestFinished();
 	OnTestFinished.Broadcast();
 
 	TObjectPtr<AActor>* ActorToDestroy = AutoDestroyActors.GetData();
