Run editor tests against the editor world, this allows using functions like UEditorActorSubsystem::GetAllLevelActors
Remove commented out FFuncTestingTickHelper from 2040597, was never used

[RN] Tests based on editor only AFunctionalTest classes now run in an editor world, rather than a PIE world

diff --git a/Engine/Source/Developer/FunctionalTesting/Classes/FunctionalTestingManager.h b/Engine/Source/Developer/FunctionalTesting/Classes/FunctionalTestingManager.h
--- a/Engine/Source/Developer/FunctionalTesting/Classes/FunctionalTestingManager.h
+++ b/Engine/Source/Developer/FunctionalTesting/Classes/FunctionalTestingManager.h
@@ -18,6 +18,13 @@ namespace FFunctionalTesting
 	extern const TCHAR* ReproStringParamsSeparator;
 }
 
+#if WITH_EDITOR
+namespace UE::Private 
+{
+	class FTickEditorTest;
+}
+#endif
+
 UCLASS(BlueprintType, MinimalAPI)
 class UFunctionalTestingManager : public UBlueprintFunctionLibrary
 {
@@ -49,8 +56,6 @@ class UFunctionalTestingManager : public UBlueprintFunctionLibrary
 	bool IsFinished() const { return bFinished; }
 	bool IsLooped() const { return bLooped; }
 	void SetLooped(const bool bNewLooped) { bLooped = bNewLooped; }
-
-	void TickMe(float DeltaTime);
 	
 protected:
 	static UFunctionalTestingManager* GetManager(UObject* WorldContext);
@@ -84,6 +89,11 @@ protected:
 	TArray<FString> TestReproStrings;
 
 private:
+
+#if WITH_EDITOR
+	TUniquePtr<UE::Private::FTickEditorTest> EditorTick;
+#endif
+
 	FTimerHandle TriggerFirstValidTestTimerHandle;
 };
 
diff --git a/Engine/Source/Developer/FunctionalTesting/Private/FunctionalTest.cpp b/Engine/Source/Developer/FunctionalTesting/Private/FunctionalTest.cpp
--- a/Engine/Source/Developer/FunctionalTesting/Private/FunctionalTest.cpp
+++ b/Engine/Source/Developer/FunctionalTesting/Private/FunctionalTest.cpp
@@ -250,7 +250,11 @@ bool AFunctionalTest::RunTest(const TArray<FString>& Params)
 {
 	SCOPE_CYCLE_COUNTER(STAT_FunctionalTest_RunTest);
 	UWorld* World = GetWorld();
+	#if WITH_EDITOR
+	ensure(World->HasBegunPlay() || (World->WorldType == EWorldType::Editor && IsEditorOnlyObject(this)));
+	#else
 	ensure(World->HasBegunPlay());
+	#endif
 	
 	FFunctionalTestBase* FunctionalTest = static_cast<FFunctionalTestBase*>(FAutomationTestFramework::Get().GetCurrentTest());
 
diff --git a/Engine/Source/Developer/FunctionalTesting/Private/FunctionalTestingManager.cpp b/Engine/Source/Developer/FunctionalTesting/Private/FunctionalTestingManager.cpp
--- a/Engine/Source/Developer/FunctionalTesting/Private/FunctionalTestingManager.cpp
+++ b/Engine/Source/Developer/FunctionalTesting/Private/FunctionalTestingManager.cpp
@@ -10,56 +10,56 @@
 #include "Engine/Engine.h"
 #include "Misc/RuntimeErrors.h"
 #include "GameFramework/PlayerController.h"
-
-#include UE_INLINE_GENERATED_CPP_BY_NAME(FunctionalTestingManager)
-
 #if WITH_EDITOR
-
-//----------------------------------------------------------------------//
-// 6/25 @todo these will be removed once marge from main comes
-class UFactory;
-//----------------------------------------------------------------------//
-
+#include "TickableEditorObject.h"
 #endif
 
+#include UE_INLINE_GENERATED_CPP_BY_NAME(FunctionalTestingManager)
+
 namespace FFunctionalTesting
 {
 	const TCHAR* ReproStringTestSeparator = TEXT("@");
 	const TCHAR* ReproStringParamsSeparator = TEXT("#");
 }
 
+#if WITH_EDITOR
+namespace UE::Private
+{
+
+class FTickEditorTest final : public FTickableEditorObject
+{
+public:
+	FTickEditorTest(TObjectPtr<AFunctionalTest> InActor)
+		: FTickableEditorObject()
+		, Actor(InActor)
+	{
+	}
+
+private:
+	virtual void Tick( float DeltaTime );
+	virtual ETickableTickType GetTickableTickType() const override 
+	{ 
+		return ETickableTickType::Always; 
+	}
+	virtual TStatId GetStatId() const override
+	{
+		RETURN_QUICK_DECLARE_CYCLE_STAT(FTickEditorTest, STATGROUP_Tickables);
+	}
 
-//struct FFuncTestingTickHelper : FTickableGameObject
-//{
-//	class UFunctionalTestingManager* Manager;
-//
-//	FFuncTestingTickHelper() : Manager(NULL) {}
-//	virtual void Tick(float DeltaTime);
-//	virtual bool IsTickable() const { return Owner && !((AActor*)Owner)->IsPendingKillPending(); }
-//	virtual bool IsTickableInEditor() const { return true; }
-//	virtual TStatId GetStatId() const ;
-//};
-//
-////----------------------------------------------------------------------//
-//// 
-////----------------------------------------------------------------------//
-//void FFuncTestingTickHelper::Tick(float DeltaTime)
-//{
-//	if (IsValid(Manager))
-//	{
-//		Manager->TickMe(DeltaTime);
-//	}
-//}
-//
-//TStatId FFuncTestingTickHelper::GetStatId() const 
-//{
-//	RETURN_QUICK_DECLARE_CYCLE_STAT(FRecastTickHelper, STATGROUP_Tickables);
-//}
-
-//----------------------------------------------------------------------//
-// 
-//----------------------------------------------------------------------//
+	TStrongObjectPtr<AFunctionalTest> Actor;
+};
+
+void FTickEditorTest::Tick(float DeltaTime)
+{
+	const TGuardValue ScriptExecutionGuard(GAllowActorScriptExecutionInEditor, true);
+	if(Actor)
+	{
+		Actor->Tick(DeltaTime);
+	}
+}
 
+}
+#endif
 
 UFunctionalTestingManager::UFunctionalTestingManager( const FObjectInitializer& ObjectInitializer )
 	: Super(ObjectInitializer)
@@ -165,11 +165,15 @@ void UFunctionalTestingManager::TriggerFirstValidTest()
 	check(World);
 	bIsRunning = true;
 
-	const bool bIsWorldInitialized =
-		World->AreActorsInitialized() &&
+	bool bIsWorldInitialized =
+		(World->AreActorsInitialized() &&
 		(!World->GetWorldSettings()->IsNavigationSystemEnabled() || !UNavigationSystemV1::IsNavigationBeingBuilt(World)) &&
-		(World->GetNumPlayerControllers() != 0) && World->GetFirstPlayerController()->GetPawnOrSpectator() != nullptr;
-
+		(World->GetNumPlayerControllers() != 0) && World->GetFirstPlayerController()->GetPawnOrSpectator() != nullptr);
+#if WITH_EDITOR
+	// editor world doesn't need to wait for navigation, controllers, pawns, or actor initialization:
+	bIsWorldInitialized |= World->WorldType == EWorldType::Editor;
+#endif
+	
 	if (bInitialDelayApplied == true && bIsWorldInitialized)
 	{
 		bIsRunning = RunFirstValidTest();
@@ -194,8 +198,11 @@ UFunctionalTestingManager* UFunctionalTestingManager::GetManager(UObject* WorldC
 	{
 		if (UWorld* World = GEngine->GetWorldFromContextObject(WorldContext, EGetWorldErrorMode::LogAndReturnNull))
 		{
-			// Only add functional test managers for the PIE and Game Worlds.
-			if (World->WorldType == EWorldType::PIE || World->WorldType == EWorldType::Game)
+			if (World->WorldType == EWorldType::PIE || World->WorldType == EWorldType::Game
+				#if WITH_EDITOR
+				|| World->WorldType == EWorldType::Editor
+				#endif
+				)
 			{
 				Manager = NewObject<UFunctionalTestingManager>(World);
 				IFunctionalTestingModule::Get().SetManager(Manager);
@@ -247,6 +254,10 @@ void UFunctionalTestingManager::OnTestDone(AFunctionalTest* FTest)
 
 void UFunctionalTestingManager::NotifyTestDone(AFunctionalTest* FTest)
 {
+#if WITH_EDITOR
+	EditorTick.Reset();
+#endif
+
 	if (FTest->OnWantsReRunCheck() == false && FTest->WantsToRunAgain() == false)
 	{
 		//We can also do named reruns. These are lower priority than those triggered above.
@@ -299,6 +310,9 @@ void UFunctionalTestingManager::AllTestsDone()
 	{
 		OnTestsComplete.Broadcast();
 		bFinished = true;
+		#if WITH_EDITOR
+		EditorTick.Reset();
+		#endif
 		IFunctionalTestingModule::Get().SetManager(nullptr);
 		RemoveFromRoot();
 	}
@@ -354,6 +368,16 @@ bool UFunctionalTestingManager::RunFirstValidTest()
 				TestToRun->TestFinishedObserver = TestFinishedObserver;
 				if (TestToRun->RunTest(TestParams))
 				{
+					#if WITH_EDITOR
+					if(!World->HasBegunPlay() && World->WorldType == EWorldType::Editor)
+					{
+						EditorTick = MakeUnique<UE::Private::FTickEditorTest>(TestToRun);
+					}
+					else
+					{
+						EditorTick.Reset();
+					}
+					#endif
 					bTestSuccessfullyTriggered = true;
 					break;
 				}
@@ -414,11 +438,6 @@ bool UFunctionalTestingManager::RunFirstValidTest()
 	return bTestSuccessfullyTriggered;
 }
 
-void UFunctionalTestingManager::TickMe(float DeltaTime)
-{
-		
-}
-
 void UFunctionalTestingManager::SetReproString(FString ReproString)
 {
 	TestReproStrings.Reset();
diff --git a/Engine/Source/Developer/FunctionalTesting/Private/FunctionalTestingModule.cpp b/Engine/Source/Developer/FunctionalTesting/Private/FunctionalTestingModule.cpp
--- a/Engine/Source/Developer/FunctionalTesting/Private/FunctionalTestingModule.cpp
+++ b/Engine/Source/Developer/FunctionalTesting/Private/FunctionalTestingModule.cpp
@@ -279,13 +279,7 @@ UWorld* FFunctionalTestingModule::GetTestWorld()
 	}
 #endif
 
-	UWorld* TestWorld = GWorld;
-	if (GIsEditor)
-	{
-		UE_LOG(LogFunctionalTest, Warning, TEXT("Functional Test using GWorld.  Not correct for PIE"));
-	}
-
-	return TestWorld;
+	return GWorld;
 }
 
 void FFunctionalTestingModule::RunAllTestsOnMap(bool bClearLog, bool bRunLooped)
diff --git a/Engine/Source/Developer/FunctionalTesting/Private/Tests/ClientFuncTestPerforming.cpp b/Engine/Source/Developer/FunctionalTesting/Private/Tests/ClientFuncTestPerforming.cpp
--- a/Engine/Source/Developer/FunctionalTesting/Private/Tests/ClientFuncTestPerforming.cpp
+++ b/Engine/Source/Developer/FunctionalTesting/Private/Tests/ClientFuncTestPerforming.cpp
@@ -11,6 +11,10 @@
 #include "FunctionalTestBase.h"
 #include "FunctionalTest.h"
 #include "FunctionalTestingHelper.h"
+#if WITH_EDITOR
+#include "Tests/AutomationEditorCommon.h"
+#include "Editor/EditorEngine.h"
+#endif
 
 #if WITH_DEV_AUTOMATION_TESTS
 
@@ -172,7 +176,40 @@ void FClientFunctionalTestingMapsEditor::GetTests(TArray<FString>& OutBeautified
 
 bool FClientFunctionalTestingMapsEditor::RunTest(const FString& Parameters)
 {
-	return FClientFunctionalTestingMapsBase::RunTest(Parameters);
+#if WITH_EDITOR
+	extern UNREALED_API UEditorEngine* GEditor;
+	FString MapObjectPath, MapPackageName, MapTestName;
+	ParseTestMapInfo(Parameters, MapObjectPath, MapPackageName, MapTestName);
+	
+	IFunctionalTestingModule::Get().MarkPendingActivation();
+
+	// Always reset these, even though tests should do the same
+	SetLogErrorAndWarningHandlingToDefault();
+
+	AddCommand(new FFunctionLatentCommand([MapTestName] {
+		TRACE_BOOKMARK(TEXT("%s ProfileBegin"), *MapTestName);
+		return true;
+	}));
+
+	AddCommand(new FOpenEditorForAssetCommand(MapPackageName));
+	AddCommand(new FFunctionLatentCommand([]() {
+		// wait for editor world
+		return GEditor->GetEditorWorldContext().World() != nullptr;
+	}));
+	
+	// run the test:
+	ADD_LATENT_AUTOMATION_COMMAND(FStartFTestOnMap(MapTestName));
+
+	// Teardown
+	AddCommand(new FFunctionLatentCommand([MapTestName] {
+		TRACE_BOOKMARK(TEXT("%s ProfileEnd"), *MapTestName);
+		return true;
+	}));
+
+	return true;
+#else
+	return false;
+#endif
 }
 
 #undef LOCTEXT_NAMESPACE
