Add override specifier to clarify intent of the tick function in this new context

#rnx

diff --git a/Engine/Source/Developer/FunctionalTesting/Private/FunctionalTestingManager.cpp b/Engine/Source/Developer/FunctionalTesting/Private/FunctionalTestingManager.cpp
--- a/Engine/Source/Developer/FunctionalTesting/Private/FunctionalTestingManager.cpp
+++ b/Engine/Source/Developer/FunctionalTesting/Private/FunctionalTestingManager.cpp
@@ -36,7 +36,7 @@ public:
 	}
 
 private:
-	virtual void Tick( float DeltaTime );
+	virtual void Tick( float DeltaTime ) override;
 	virtual ETickableTickType GetTickableTickType() const override 
 	{ 
 		return ETickableTickType::Always; 
