Add temporary logging to determine cause of editor test hang in the farm

#rnx

diff --git a/Engine/Source/Developer/FunctionalTesting/Private/Tests/ClientFuncTestPerforming.cpp b/Engine/Source/Developer/FunctionalTesting/Private/Tests/ClientFuncTestPerforming.cpp
--- a/Engine/Source/Developer/FunctionalTesting/Private/Tests/ClientFuncTestPerforming.cpp
+++ b/Engine/Source/Developer/FunctionalTesting/Private/Tests/ClientFuncTestPerforming.cpp
@@ -224,6 +224,7 @@ bool FClientFunctionalTestingMapsEditor::RunTest(const FString& Parameters)
 	AddCommand(new FOpenEditorForAssetCommand(MapPackageName));
 	AddCommand(new FFunctionLatentCommand([]() {
 		// wait for editor world
+		UE_LOG(LogTemp, Log, TEXT("Waiting for editor world"));
 		return GEditor->GetEditorWorldContext().World() != nullptr;
 	}));
 	
