Run tests based on non-editor only classes that identify themselves as EditorOnly in PIE rather than in vanilla editor
#rnx

diff --git a/Engine/Source/Developer/FunctionalTesting/Private/Tests/ClientFuncTestPerforming.cpp b/Engine/Source/Developer/FunctionalTesting/Private/Tests/ClientFuncTestPerforming.cpp
--- a/Engine/Source/Developer/FunctionalTesting/Private/Tests/ClientFuncTestPerforming.cpp
+++ b/Engine/Source/Developer/FunctionalTesting/Private/Tests/ClientFuncTestPerforming.cpp
@@ -203,7 +203,7 @@ bool FClientFunctionalTestingMapsEditor::RunTest(const FString& Parameters)
 	{
 		// we have a world to use - but does it have the test we now need to run?
 		AFunctionalTest* Test = FindObjectFast<AFunctionalTest>(WorldToUse->PersistentLevel, FName(MapTestName));
-		if(Test && (!IsEditorOnlyObject(Test) || Test->IsEditorOnlyLoadedInPIE()))
+		if(Test && (!IsEditorOnlyObject(Test) || Test->IsEditorOnlyLoadedInPIE() || !IsEditorOnlyObject(Test->GetClass())))
 		{
 			// it does, just use base test running logic, as it will be faster than
 			// reopening the map:
@@ -216,11 +216,6 @@ bool FClientFunctionalTestingMapsEditor::RunTest(const FString& Parameters)
 	// Always reset these, even though tests should do the same
 	SetLogErrorAndWarningHandlingToDefault();
 
-	AddCommand(new FFunctionLatentCommand([MapTestName] {
-		TRACE_BOOKMARK(TEXT("%s ProfileBegin"), *MapTestName);
-		return true;
-	}));
-
 	AddCommand(new FOpenEditorForAssetCommand(MapPackageName));
 	AddCommand(new FFunctionLatentCommand([]() {
 		// wait for editor world
@@ -233,7 +228,9 @@ bool FClientFunctionalTestingMapsEditor::RunTest(const FString& Parameters)
 		AddCommand(new FFunctionLatentCommand([MapTestName, MapPackageName]() {
 			UWorld* World = GEditor->GetEditorWorldContext().World();
 			AFunctionalTest* Test = FindObjectFast<AFunctionalTest>(World->PersistentLevel, FName(MapTestName));
-			if(ensure(Test) && Test->IsEditorOnlyLoadedInPIE())
+			// Actors that indicate they are editor only but want IsNonPIEEditorOnly or are of a class
+			// that is not editor only should run in PIE:
+			if(ensure(Test) && (Test->IsEditorOnlyLoadedInPIE() || !IsEditorOnlyObject(Test->GetClass())))
 			{
 				// request PIE, as this test actor claims it supports PIE:
 				const bool bCanProceed = AutomationOpenMap(MapPackageName);
@@ -246,15 +243,13 @@ bool FClientFunctionalTestingMapsEditor::RunTest(const FString& Parameters)
 		}));
 	}
 
-	// run the test:
+	// run the test - note that FStartFTestOnMap is going to 
+	// create two more latent commands in its update function,
+	// meaning any more latent commands you add will run
+	// before FStartFTestOnMap actually does anything (as it
+	// will simply enqueue commands after your command)
 	ADD_LATENT_AUTOMATION_COMMAND(FStartFTestOnMap(MapTestName));
-
-	// Teardown
-	AddCommand(new FFunctionLatentCommand([MapTestName] {
-		TRACE_BOOKMARK(TEXT("%s ProfileEnd"), *MapTestName);
-		return true;
-	}));
-
+	// !!! ANY COMMANDS ADDED HERE WILL RUN BEFORE FStartFTestOnMap !!!
 	return true;
 #else
 	// This is likely unreachable, but because this is tag based we can't guarantee that no one is using 
