Avoid applying initial delay for actor tests that can run in the editor world

#rnx

#virtualized

diff --git a/Engine/Source/Developer/FunctionalTesting/Private/FunctionalTestingManager.cpp b/Engine/Source/Developer/FunctionalTesting/Private/FunctionalTestingManager.cpp
--- a/Engine/Source/Developer/FunctionalTesting/Private/FunctionalTestingManager.cpp
+++ b/Engine/Source/Developer/FunctionalTesting/Private/FunctionalTestingManager.cpp
@@ -170,8 +170,17 @@ void UFunctionalTestingManager::TriggerFirstValidTest()
 		(!World->GetWorldSettings()->IsNavigationSystemEnabled() || !UNavigationSystemV1::IsNavigationBeingBuilt(World)) &&
 		(World->GetNumPlayerControllers() != 0) && World->GetFirstPlayerController()->GetPawnOrSpectator() != nullptr);
 #if WITH_EDITOR
+	const bool bInEditorWorld = (World->WorldType == EWorldType::Editor);
 	// editor world doesn't need to wait for navigation, controllers, pawns, or actor initialization:
-	bIsWorldInitialized |= World->WorldType == EWorldType::Editor;
+	bIsWorldInitialized |= bInEditorWorld;
+	bInitialDelayApplied |= bInEditorWorld; // skip initial delay in editor world, wastes time and would require a different timer manager
+	// The ability to run actor based tests in the editor world is new
+	// so I'm adding a logging statement so that we have a fingerprint
+	// in the log of any behavior change:
+	if(bInEditorWorld)
+	{
+		UE_LOG(LogTemp, Log, TEXT("Running %s in Editor World"), TestReproStrings.Num() > 0 ? *TestReproStrings[0] : TEXT("unknown test"));
+	}
 #endif
 	
 	if (bInitialDelayApplied == true && bIsWorldInitialized)
