Fix for a crash with FunctionalTest TestTags that was occurring because the level is sometimes NULL when modifying the default blueprint class

diff --git a/Engine/Source/Developer/FunctionalTesting/Private/FunctionalTest.cpp b/Engine/Source/Developer/FunctionalTesting/Private/FunctionalTest.cpp
--- a/Engine/Source/Developer/FunctionalTesting/Private/FunctionalTest.cpp
+++ b/Engine/Source/Developer/FunctionalTesting/Private/FunctionalTest.cpp
@@ -631,15 +631,22 @@ void AFunctionalTest::PostEditChangeProperty( struct FPropertyChangedEvent& Prop
 			}
 			else if (PropertyChangedEvent.Property->GetFName() == NAME_TestTags)
 			{
-				// re-register existing tags for the loaded test
-				FString CurrentMapPackageName = this->GetLevel()->GetPackage()->GetName();
-				FString PartialSuiteName = MapPackageToAutomationPath(CurrentMapPackageName);
-				FString FullBeautifiedName(PartialSuiteName + TEXT(".") + *GetActorLabel());
-				FString FullTestName(TEXT("Project.Functional Tests." + FullBeautifiedName));
-
-				FAutomationTestFramework& TestFramework = FAutomationTestFramework::Get();
-				TestFramework.UnregisterAutomationTestTags(FullTestName);
-				TestFramework.RegisterAutomationTestTags(FullTestName, TestTags);
+				if (const ULevel* Level = GetLevel())
+				{
+					if (const UPackage* Package = Level->GetPackage())
+					{
+						// re-register existing tags for the loaded test
+						const FString CurrentMapPackageName = Package->GetName();
+						const FString PartialSuiteName = MapPackageToAutomationPath(CurrentMapPackageName);
+						const FString FullBeautifiedName(PartialSuiteName + TEXT(".") + *GetActorLabel());
+						const FString FullTestName(TEXT("Project.Functional Tests." + FullBeautifiedName));
+
+						FAutomationTestFramework& TestFramework = FAutomationTestFramework::Get();
+						TestFramework.UnregisterAutomationTestTags(FullTestName);
+						TestFramework.RegisterAutomationTestTags(FullTestName, TestTags);
+					}
+				}
+				
 			}
 		}
 	}
