PR #12350: Allow specifying if a functional test is in a sublevel and in which persistent level is should be enabled.

#virtualized

diff --git a/Engine/Source/Developer/FunctionalTesting/Classes/FunctionalTest.h b/Engine/Source/Developer/FunctionalTesting/Classes/FunctionalTest.h
--- a/Engine/Source/Developer/FunctionalTesting/Classes/FunctionalTest.h
+++ b/Engine/Source/Developer/FunctionalTesting/Classes/FunctionalTest.h
@@ -281,6 +281,18 @@ protected:
 	UPROPERTY(EditAnywhere, BlueprintReadOnly, Category="Functional Testing", meta=(ScriptName="IsEnabledValue"))
 	uint32 bIsEnabled:1;
 
+	/**
+	 * Allows a test to be enabled only if loaded from the specified persistent level; avoiding other levels to run the test.
+	 */
+	UPROPERTY(EditInstanceOnly, Category="Functional Testing")
+	uint32 bIsInSublevel:1;
+
+	/**
+	 * Name of the persistent level to run the test from.
+	 */
+	UPROPERTY(EditInstanceOnly, Category="Functional Testing", meta=(EditCondition="bIsInSublevel"))
+	FName PersistentLevelName;
+
 	/**
 	 * Determines how LogErrors are handled during this test.
 	 */
@@ -732,6 +744,9 @@ public:
 	UFUNCTION(BlueprintCallable, Category = "Functional Testing")
 	bool IsEnabled() const;
 
+	UFUNCTION(BlueprintCallable, Category = "Functional Testing")
+	bool IsEnabledInWorld(const UWorld* World) const;
+
 protected:
 	/**
 	 * Prepare Test is fired once the test starts up, before the test IsReady() and thus before Start Test is called.
diff --git a/Engine/Source/Developer/FunctionalTesting/Private/FunctionalTest.cpp b/Engine/Source/Developer/FunctionalTesting/Private/FunctionalTest.cpp
--- a/Engine/Source/Developer/FunctionalTesting/Private/FunctionalTest.cpp
+++ b/Engine/Source/Developer/FunctionalTesting/Private/FunctionalTest.cpp
@@ -138,6 +138,8 @@ AFunctionalTest::AFunctionalTest( const FObjectInitializer& ObjectInitializer )
 	: Super(ObjectInitializer)
 	, TestLabel(GetName())
 	, bIsEnabled(true)
+	, bIsInSublevel(false)
+	, PersistentLevelName()
 	, LogErrorHandling(EFunctionalTestLogHandling::ProjectDefault)
 	, LogWarningHandling(EFunctionalTestLogHandling::ProjectDefault)
 	, bShouldDelayGarbageCollection(true)
@@ -509,6 +511,18 @@ bool AFunctionalTest::IsEnabled() const
 	return bIsEnabled;
 }
 
+bool AFunctionalTest::IsEnabledInWorld(const UWorld* World) const
+{
+	if (bIsInSublevel)
+	{
+		return World != nullptr && World->GetFName() == PersistentLevelName
+			? bIsEnabled
+			: false;
+	}
+
+	return bIsEnabled;
+}
+
 //@todo add "warning" level here
 void AFunctionalTest::LogMessage(const FString& Message)
 {
diff --git a/Engine/Source/Developer/FunctionalTesting/Private/FunctionalTestingModule.cpp b/Engine/Source/Developer/FunctionalTesting/Private/FunctionalTestingModule.cpp
--- a/Engine/Source/Developer/FunctionalTesting/Private/FunctionalTestingModule.cpp
+++ b/Engine/Source/Developer/FunctionalTesting/Private/FunctionalTestingModule.cpp
@@ -79,7 +79,7 @@ void FFunctionalTestingModule::OnGetAssetTagsForWorld(const UWorld* World, FAsse
 		if (!FunctionalTest->IsPackageExternal())
 		{
 			// Only include enabled tests in the list of functional tests to run.
-			if (FunctionalTest->IsEnabled())
+			if (FunctionalTest->IsEnabledInWorld(World))
 			{
 				TArray<FString>& TestNames = IsEditorOnlyObject(FunctionalTest) ? TestNamesEditor : TestNamesRuntime;
 				TestNames.Add(FString::Printf(TEXT("%s|%s|%s;"), *FunctionalTest->GetActorLabel(), *FunctionalTest->GetName(), *FunctionalTest->TestTags));
