Functional Tests: Add support for custom categorization.

diff --git a/Engine/Source/Developer/FunctionalTesting/Classes/FunctionalTest.h b/Engine/Source/Developer/FunctionalTesting/Classes/FunctionalTest.h
--- a/Engine/Source/Developer/FunctionalTesting/Classes/FunctionalTest.h
+++ b/Engine/Source/Developer/FunctionalTesting/Classes/FunctionalTest.h
@@ -211,6 +211,12 @@ FString FUNCTIONALTESTING_API LexToString(const EFunctionalTestResult TestResult
 /** Return a dot-separated path prefix string representing the map that contains a test */
 FString FUNCTIONALTESTING_API MapPackageToAutomationPath(const FString& MapPackageName);
 
+/**
+ * Return a dot-separated path prefix string representing the map that contains a test.
+ * If the map defines a "TestPathOverride" AR tag, it'll take its value instead of building the default prefix.
+ */
+FString FUNCTIONALTESTING_API MapPackageToAutomationPath(const FAssetData& MapAsset);
+
 /** Registration information for an individual test */
 struct FFunctionalTestInfo
 {
diff --git a/Engine/Source/Developer/FunctionalTesting/Private/FunctionalTest.cpp b/Engine/Source/Developer/FunctionalTesting/Private/FunctionalTest.cpp
--- a/Engine/Source/Developer/FunctionalTesting/Private/FunctionalTest.cpp
+++ b/Engine/Source/Developer/FunctionalTesting/Private/FunctionalTest.cpp
@@ -2,6 +2,8 @@
 
 #include "FunctionalTest.h"
 #include "FunctionalTestingModule.h"
+#include "AssetRegistry/AssetData.h"
+#include "AssetRegistry/IAssetRegistry.h"
 #include "GameFramework/Pawn.h"
 #include "Misc/Paths.h"
 #include "Engine/GameViewportClient.h"
@@ -134,6 +136,20 @@ FString MapPackageToAutomationPath(const FString& MapPackageName)
 }
 
 
+FString MapPackageToAutomationPath(const FAssetData& MapAsset)
+{
+	static const FName TestPathOverrideName = "TestPathOverride";
+
+	FString StoredPathName;
+	if (MapAsset.GetTagValue(TestPathOverrideName, StoredPathName))
+	{
+		return StoredPathName;
+	}
+
+	return MapPackageToAutomationPath(MapAsset);
+}
+
+
 AFunctionalTest::AFunctionalTest( const FObjectInitializer& ObjectInitializer )
 	: Super(ObjectInitializer)
 	, TestLabel(GetName())
@@ -650,8 +666,9 @@ void AFunctionalTest::PostEditChangeProperty( struct FPropertyChangedEvent& Prop
 					if (const UPackage* Package = Level->GetPackage())
 					{
 						// re-register existing tags for the loaded test
-						const FString CurrentMapPackageName = Package->GetName();
-						const FString PartialSuiteName = MapPackageToAutomationPath(CurrentMapPackageName);
+						FSoftObjectPath CurrentMapPath = FSoftObjectPath(this->GetLevel()->GetPackage());
+						FAssetData MapAsset = IAssetRegistry::Get()->GetAssetByObjectPath(CurrentMapPath);
+						const FString PartialSuiteName = MapPackageToAutomationPath(MapAsset);
 						const FString FullBeautifiedName(PartialSuiteName + TEXT(".") + *GetActorLabel());
 						const FString FullTestName(TEXT("Project.Functional Tests." + FullBeautifiedName));
 
diff --git a/Engine/Source/Developer/FunctionalTesting/Private/FunctionalTestingModule.cpp b/Engine/Source/Developer/FunctionalTesting/Private/FunctionalTestingModule.cpp
--- a/Engine/Source/Developer/FunctionalTesting/Private/FunctionalTestingModule.cpp
+++ b/Engine/Source/Developer/FunctionalTesting/Private/FunctionalTestingModule.cpp
@@ -145,7 +145,7 @@ void FFunctionalTestingModule::GetMapTests(bool bEditorOnlyTests, TArray<FFuncti
 				FString MapAssetPath = MapAsset.GetObjectPathString();
 				FString MapPackageName = MapAsset.PackageName.ToString();
 				if (!IsDeveloperDirectoryIncluded && MapPackageName.Find(TEXT("/Game/Developers")) == 0) continue;
-				FString PartialSuiteName = MapPackageToAutomationPath(MapPackageName);
+				FString PartialSuiteName = MapPackageToAutomationPath(MapAsset);
 
 				TArray<FString> MapTests;
 				FAssetDataTagMapSharedView::FFindTagResult MapAssetNamesResult = MapAsset.TagsAndValues.FindTag(bEditorOnlyTests ? TEXT("TestNamesEditor") : TEXT("TestNames"));
