Fix screenshot functional tests on mobile XR devices.
AScreenshotFunctionalTest uses HighResShot with a fixed resolution by default on mobile, this isn't suitable for XR devices, since forcing a new resolution on the existing swapchain results in corruption. The intended result (being able to compare screenshots between device models) also isn't viable in XR due to varying FOVs.
We also don't want to lock the test camera to the HMD, since positional tracking will result in inconsistent screenshots.

diff --git a/Engine/Source/Developer/FunctionalTesting/Private/ScreenshotFunctionalTest.cpp b/Engine/Source/Developer/FunctionalTesting/Private/ScreenshotFunctionalTest.cpp
--- a/Engine/Source/Developer/FunctionalTesting/Private/ScreenshotFunctionalTest.cpp
+++ b/Engine/Source/Developer/FunctionalTesting/Private/ScreenshotFunctionalTest.cpp
@@ -93,7 +93,7 @@ void AScreenshotFunctionalTest::RequestScreenshot()
 {
 	Super::RequestScreenshot();
 
-	if(IsMobilePlatform(GShaderPlatformForFeatureLevel[GMaxRHIFeatureLevel]))
+	if(IsMobilePlatform(GShaderPlatformForFeatureLevel[GMaxRHIFeatureLevel]) && !GEngine->IsStereoscopic3D())
 	{
 		// For mobile, use the high res screenshot API to ensure a fixed resolution screenshot is produced.
 		// This means screenshot comparisons can compare with the output from any device.
diff --git a/Engine/Source/Developer/FunctionalTesting/Private/ScreenshotFunctionalTestBase.cpp b/Engine/Source/Developer/FunctionalTesting/Private/ScreenshotFunctionalTestBase.cpp
--- a/Engine/Source/Developer/FunctionalTesting/Private/ScreenshotFunctionalTestBase.cpp
+++ b/Engine/Source/Developer/FunctionalTesting/Private/ScreenshotFunctionalTestBase.cpp
@@ -38,6 +38,7 @@ AScreenshotFunctionalTestBase::AScreenshotFunctionalTestBase(const FObjectInitia
 {
 	ScreenshotCamera = CreateDefaultSubobject<UCameraComponent>(TEXT("Camera"));
 	ScreenshotCamera->SetupAttachment(RootComponent);
+	ScreenshotCamera->bLockToHmd = false;
 
 #if WITH_AUTOMATION_TESTS
 	ScreenshotEnvSetup = MakeShareable(new FAutomationTestScreenshotEnvSetup());
