Update functional screenshot tests to support emulating splitscreen when taking comparison screenshots.
This mode should not require registering new ground truths in order to pass.
Splitscreen emulation works by temporarily causing one player to register two views, each of which is rendered separately at full size, then taking a screenshot of only the second view.

r.ScreenshotTest.EmulateSplitscreenMode:
0 - Off (default)
1 - Horizontal Layout
2 - Vertical Layout
3 - Diagonal Layout

Individual tests can opt out of splitscreen emulation by unchecking "Allow Emulating Splitscreen" under "Screenshot Options".

diff --git a/Engine/Source/Developer/FunctionalTesting/Private/ScreenshotFunctionalTest.cpp b/Engine/Source/Developer/FunctionalTesting/Private/ScreenshotFunctionalTest.cpp
--- a/Engine/Source/Developer/FunctionalTesting/Private/ScreenshotFunctionalTest.cpp
+++ b/Engine/Source/Developer/FunctionalTesting/Private/ScreenshotFunctionalTest.cpp
@@ -98,8 +98,8 @@ void AScreenshotFunctionalTest::RequestScreenshot()
 		// For mobile, use the high res screenshot API to ensure a fixed resolution screenshot is produced.
 		// This means screenshot comparisons can compare with the output from any device.
 		FHighResScreenshotConfig& Config = GetHighResScreenshotConfig();
-		FIntPoint ScreenshotViewportSize = UAutomationBlueprintFunctionLibrary::GetAutomationScreenshotSize(ScreenshotOptions);
-		if (Config.SetResolution(ScreenshotViewportSize.X, ScreenshotViewportSize.Y, 1.0f))
+		FIntPoint Size = UAutomationBlueprintFunctionLibrary::GetAutomationScreenshotSize(ScreenshotOptions);
+		if (Config.SetResolution(Size.X, Size.Y, 1.0f))
 		{
 			UGameViewportClient* GameViewportClient = AutomationCommon::GetAnyGameViewportClient();
 			check(GameViewportClient);
@@ -118,8 +118,36 @@ void AScreenshotFunctionalTest::RequestScreenshot()
 		//    FEditorViewportClient::ProcessScreenshots is called to read pixels back from the viewport. It also
 		//    trigger the callback function registered in step 2.
 
-		bool bShowUI = false;
-		FScreenshotRequest::RequestScreenshot(bShowUI);
+		FIntRect ViewRect = FIntRect(); // Passing an empty ViewRect defaults to including to the entire render target in the screenshot
+
+		if (ScreenshotOptions.bAllowEmulatingSplitscreen)
+		{
+			static const TConsoleVariableData<int32>* CVarEmulateSplitscreenMode = IConsoleManager::Get().FindTConsoleVariableDataInt(TEXT("r.ScreenshotTest.EmulateSplitscreenMode"));
+			const int32 EmulateSplitscreenMode = CVarEmulateSplitscreenMode ? CVarEmulateSplitscreenMode->GetValueOnAnyThread() : 0;
+			const FIntPoint Size = UAutomationBlueprintFunctionLibrary::GetAutomationScreenshotSize(ScreenshotOptions);	
+
+			switch (EmulateSplitscreenMode)
+			{
+				case 0:
+					break;
+				case 1:	// Horizontal
+					ViewRect = FIntRect(Size.X, 0, Size.X * 2, Size.Y);
+					break;
+				case 2:	// Vertical
+					ViewRect = FIntRect(0, Size.Y, Size.X, Size.Y * 2);
+					break;
+				case 3:	// Diagonal
+					ViewRect = FIntRect(Size, Size * 2);
+					break;
+				default:
+					// Already hit ensure in AScreenshotFunctionalTestBase::PrepareForScreenshot()
+					break;
+			}
+
+			checkf(ViewRect.Size() == Size || ViewRect.IsEmpty(), TEXT("Screenshot view rect should either match the expected automation screenshot size, or be empty (which requests a screenshot of the entire render target)."));
+		}
+		
+		FScreenshotRequest::RequestScreenshot(TEXT(""), /*bInShowUI*/ false, /*bAddUniqueSuffix*/ true, /*bHdrScreenshot*/ false, ViewRect);
 	}
 }
 
diff --git a/Engine/Source/Developer/FunctionalTesting/Private/ScreenshotFunctionalTestBase.cpp b/Engine/Source/Developer/FunctionalTesting/Private/ScreenshotFunctionalTestBase.cpp
--- a/Engine/Source/Developer/FunctionalTesting/Private/ScreenshotFunctionalTestBase.cpp
+++ b/Engine/Source/Developer/FunctionalTesting/Private/ScreenshotFunctionalTestBase.cpp
@@ -19,6 +19,7 @@
 #include "Logging/LogMacros.h"
 #include "UObject/AutomationObjectVersion.h"
 #include "RenderGraphBuilder.h"
+#include "Engine/LocalPlayer.h"
 
 #include UE_INLINE_GENERATED_CPP_BY_NAME(ScreenshotFunctionalTestBase)
 
@@ -29,6 +30,18 @@ static TAutoConsoleVariable<int32> GDumpGPUDumpOnScreenshotTest(
 	TEXT("Allows to filter the tree when using r.DumpGPU command, the pattern match is case sensitive."),
 	ECVF_Default);
 
+static TAutoConsoleVariable<int32> CVarEmulateSplitscreenMode(
+	TEXT("r.ScreenshotTest.EmulateSplitscreenMode"), 0,
+	TEXT("Allows automtically emulating splitscreen rendering for all screenshot functional tests. Useful for catching splitscreen exclusive rendering issues.\n")
+	TEXT("When active and bAllowEmulatingSplitscreen=true for a test, two full size views will be rendered, and screenshot comparison tests will capture only the second.\n")
+	TEXT("Tests with existing ground truths should pass if there are no splitscreen rendering issues.")
+	TEXT("If a test is failing in splitscreen due to differences that do not impact correctness (such as changes in noise patterns), you can exempt it by setting \"Allow Emulating Splitscren\" to false under the Screenshot Options.\n")
+	TEXT("0 - Off\n")
+	TEXT("1 - Horizontal: Viewports are side by side\n")
+	TEXT("2 - Vertical: Viewports are at the top and bottom\n")
+	TEXT("3 - Diagonal: Viewports are in the upper left and lower right, leaving gaps in the other corners of the render target"),
+	ECVF_Default);
+
 AScreenshotFunctionalTestBase::AScreenshotFunctionalTestBase(const FObjectInitializer& ObjectInitializer)
 	: AFunctionalTest(ObjectInitializer)
 	, ScreenshotOptions(EComparisonTolerance::Low)
@@ -58,6 +71,15 @@ void AScreenshotFunctionalTestBase::PrepareTest()
 		// Make sure the camera target is not auto managed
 		PlayerController->bAutoManageActiveCameraTarget = false;
 		PlayerController->SetViewTarget(this, FViewTargetTransitionParams());
+
+		if (ScreenshotOptions.bAllowEmulatingSplitscreen)
+		{
+			const int32 EmulateSplitscreenMode = CVarEmulateSplitscreenMode.GetValueOnAnyThread();
+			if (EmulateSplitscreenMode >= 1 && EmulateSplitscreenMode <= 3)
+			{
+				PlayerController->GetLocalPlayer()->bEmulateSplitscreen = true;
+			}
+		}
 	}
 
 	PrepareForScreenshot();
@@ -143,6 +165,28 @@ void AScreenshotFunctionalTestBase::PrepareForScreenshot()
 		{
 			ViewportRestoreSize = GameViewport->GetSize();
 			FIntPoint ScreenshotViewportSize = UAutomationBlueprintFunctionLibrary::GetAutomationScreenshotSize(ScreenshotOptions);
+
+			if (ScreenshotOptions.bAllowEmulatingSplitscreen)
+			{
+				switch (CVarEmulateSplitscreenMode->GetInt())
+				{
+					case 0:
+						break;
+					case 1:	// Horizontal
+						ScreenshotViewportSize.X *= 2;
+						break;
+					case 2:	// Vertical
+						ScreenshotViewportSize.Y *= 2;
+						break;
+					case 3:	// Diagonal
+						ScreenshotViewportSize *= 2;
+						break;
+					default:
+						ensureMsgf(false, TEXT("Selected r.ScreenshotTest.EmulateSplitscreenMode of %d is invalid. Defaulting to 0 (off)."), CVarEmulateSplitscreenMode->GetInt());
+						break;
+				}
+			}
+
 			GameViewport->SetViewportSize(ScreenshotViewportSize.X, ScreenshotViewportSize.Y);
 			bNeedsViewportRestore = true;
 		}
diff --git a/Engine/Source/Developer/FunctionalTesting/Public/AutomationScreenshotOptions.h b/Engine/Source/Developer/FunctionalTesting/Public/AutomationScreenshotOptions.h
--- a/Engine/Source/Developer/FunctionalTesting/Public/AutomationScreenshotOptions.h
+++ b/Engine/Source/Developer/FunctionalTesting/Public/AutomationScreenshotOptions.h
@@ -86,6 +86,7 @@ public:
 		, MaximumGlobalError(0.02f)
 		, bIgnoreAntiAliasing(true)
 		, bIgnoreColors(false)
+		, bAllowEmulatingSplitscreen(true)
 	{
 	}
 
@@ -105,6 +106,7 @@ public:
 		, MaximumGlobalError(0.02f)
 		, bIgnoreAntiAliasing(true)
 		, bIgnoreColors(false)
+		, bAllowEmulatingSplitscreen(true)
 	{
 		SetToleranceAmounts(InTolerance);
 	}
@@ -216,6 +218,12 @@ public:
 	UPROPERTY(EditAnywhere, BlueprintReadWrite, Category="Comparison", AdvancedDisplay)
 	bool bIgnoreColors;
 
+	/**
+	 * If this is true, this test can use emulated splitscreen rendering when r.ScreenshotTest.EmulateSplitscreenMode = 1, 2, or 3.
+	 */
+	UPROPERTY(EditAnywhere, BlueprintReadWrite, Category="Comparison")
+	bool bAllowEmulatingSplitscreen;
+
 public:
 	void SetToleranceAmounts(EComparisonTolerance InTolerance)
 	{
diff --git a/Engine/Source/Runtime/Engine/Classes/Engine/LocalPlayer.h b/Engine/Source/Runtime/Engine/Classes/Engine/LocalPlayer.h
--- a/Engine/Source/Runtime/Engine/Classes/Engine/LocalPlayer.h
+++ b/Engine/Source/Runtime/Engine/Classes/Engine/LocalPlayer.h
@@ -212,6 +212,9 @@ public:
 	/** Event called when this local player has had a new outer PlayerController set */
 	DECLARE_EVENT_OneParam(ULocalPlayer, FOnPlayerControllerChanged, APlayerController* /*NewPC*/);
 	FOnPlayerControllerChanged& OnPlayerControllerChanged() { return OnPlayerControllerChangedEvent; }
+
+	/** Whether or not we're emulating splitscreen with this player. Should only be set from screenshot functional tests. */
+	bool bEmulateSplitscreen = false;
 	
 	/** 
 	 * Set the connection identifier for this local user. 
diff --git a/Engine/Source/Runtime/Engine/Private/GameViewportClient.cpp b/Engine/Source/Runtime/Engine/Private/GameViewportClient.cpp
--- a/Engine/Source/Runtime/Engine/Private/GameViewportClient.cpp
+++ b/Engine/Source/Runtime/Engine/Private/GameViewportClient.cpp
@@ -1672,8 +1672,11 @@ void UGameViewportClient::Draw(FViewport* InViewport, FCanvas* SceneCanvas)
 		{
 			APlayerController* PlayerController = LocalPlayer->PlayerController;
 
-			const bool bEnableStereo = GEngine->IsStereoscopic3D(InViewport);
-			const int32 NumViews = bStereoRendering ? GEngine->StereoRenderingDevice->GetDesiredNumberOfViews(bStereoRendering) : 1;
+			const int32 NumViews = bStereoRendering ? 
+				GEngine->StereoRenderingDevice->GetDesiredNumberOfViews(bStereoRendering) 
+				: (LocalPlayer->bEmulateSplitscreen ? 2 : 1);
+
+			checkf(!bStereoRendering || !LocalPlayer->bEmulateSplitscreen, TEXT("r.ScreenshotTest.EmulateSplitscreenMode is unsupported when using stereo rendering."));
 
 			for (int32 ViewIndex = 0; ViewIndex < NumViews; ++ViewIndex)
 			{
@@ -1681,7 +1684,7 @@ void UGameViewportClient::Draw(FViewport* InViewport, FCanvas* SceneCanvas)
 				FVector		ViewLocation;
 				FRotator	ViewRotation;
 
-				FSceneView* View = LocalPlayer->CalcSceneView(&ViewFamily, ViewLocation, ViewRotation, InViewport, nullptr, bStereoRendering ? ViewIndex : INDEX_NONE);
+				FSceneView* View = LocalPlayer->CalcSceneView(&ViewFamily, ViewLocation, ViewRotation, InViewport, nullptr, (bStereoRendering || LocalPlayer->bEmulateSplitscreen) ? ViewIndex : INDEX_NONE);
 
 				if (View)
 				{
diff --git a/Engine/Source/Runtime/Engine/Private/LocalPlayer.cpp b/Engine/Source/Runtime/Engine/Private/LocalPlayer.cpp
--- a/Engine/Source/Runtime/Engine/Private/LocalPlayer.cpp
+++ b/Engine/Source/Runtime/Engine/Private/LocalPlayer.cpp
@@ -758,6 +758,37 @@ bool ULocalPlayer::CalcSceneViewInitOptions(
 	{
 		return false;
 	}
+
+	// If we are emulating splitscreen rendering for a screenshot functional test, override the layout set in UGameViewportClient::LayoutPlayers()
+	if (bEmulateSplitscreen)
+	{
+		static const TConsoleVariableData<int32>* CVarEmulateSplitscreenMode = IConsoleManager::Get().FindTConsoleVariableDataInt(TEXT("r.ScreenshotTest.EmulateSplitscreenMode"));
+		const int32 EmulateSplitscreenMode = CVarEmulateSplitscreenMode ? CVarEmulateSplitscreenMode->GetValueOnAnyThread() : 0;
+
+		checkf(EmulateSplitscreenMode >= 1 && EmulateSplitscreenMode <= 3, TEXT("bEmulateSplitscreen should only be set by screenshot functional tests via r.ScreenshotTest.EmulateSplitscreenMode"));
+		checkf((StereoViewIndex == 0 || StereoViewIndex == 1), TEXT("bEmulateSplitscreen only supports two splitscreen views"));
+
+		const bool bBaseView = (StereoViewIndex == 0);
+		switch (EmulateSplitscreenMode)
+		{
+			case 1:	// Horizontal
+				Origin = bBaseView ? FVector2D(0.0f, 0.0f) : FVector2D(0.5f, 0.0f);
+				Size = FVector2D(0.5f, 1.0f);
+				break;
+			case 2:	// Vertical
+				Origin = bBaseView ? FVector2D(0.0f, 0.0f) : FVector2D(0.0f, 0.5f);
+				Size = FVector2D(1.0f, 0.5f);
+				break;
+			case 3:	// Diagonal
+				Origin = bBaseView ? FVector2D(0.0f, 0.0f) : FVector2D(0.5f, 0.5f);
+				Size = FVector2D(0.5f, 0.5f);
+				break;
+			default:
+				checkNoEntry();
+				break;
+		}
+	}
+
 	// get the projection data
 	if (GetProjectionData(Viewport, /*inout*/ ViewInitOptions, StereoViewIndex) == false)
 	{
