Add warning for screenshot functional test splitscreen emulation when the screenshot size does not match the quantized scene buffer size.

#Synced-CL 52477026

diff --git a/Engine/Source/Developer/FunctionalTesting/Private/ScreenshotFunctionalTestBase.cpp b/Engine/Source/Developer/FunctionalTesting/Private/ScreenshotFunctionalTestBase.cpp
--- a/Engine/Source/Developer/FunctionalTesting/Private/ScreenshotFunctionalTestBase.cpp
+++ b/Engine/Source/Developer/FunctionalTesting/Private/ScreenshotFunctionalTestBase.cpp
@@ -20,6 +20,7 @@
 #include "UObject/AutomationObjectVersion.h"
 #include "RenderGraphBuilder.h"
 #include "Engine/LocalPlayer.h"
+#include "RenderUtils.h"
 
 #include UE_INLINE_GENERATED_CPP_BY_NAME(ScreenshotFunctionalTestBase)
 
@@ -168,7 +169,21 @@ void AScreenshotFunctionalTestBase::PrepareForScreenshot()
 
 			if (ScreenshotOptions.bAllowEmulatingSplitscreen)
 			{
-				switch (CVarEmulateSplitscreenMode->GetInt())
+				const int32 EmulateSplitscreenMode = CVarEmulateSplitscreenMode->GetInt();
+
+				if (EmulateSplitscreenMode >= 1 && EmulateSplitscreenMode <= 3)
+				{
+					FIntPoint QuantizedScreenshotViewportSize;
+					QuantizeSceneBufferSize(ScreenshotViewportSize, QuantizedScreenshotViewportSize);
+
+					if (ScreenshotViewportSize != QuantizedScreenshotViewportSize)
+					{
+						UE_LOG(LogScreenshotFunctionalTest, Warning, TEXT("Screenshot size %s does not match quantized screenshot size %s, emulated splitscreen rendering may behave incorrectly."),
+							*ScreenshotViewportSize.ToString(), *QuantizedScreenshotViewportSize.ToString());
+					}
+				}
+
+				switch (EmulateSplitscreenMode)
 				{
 					case 0:
 						break;
@@ -182,7 +197,7 @@ void AScreenshotFunctionalTestBase::PrepareForScreenshot()
 						ScreenshotViewportSize *= 2;
 						break;
 					default:
-						ensureMsgf(false, TEXT("Selected r.ScreenshotTest.EmulateSplitscreenMode of %d is invalid. Defaulting to 0 (off)."), CVarEmulateSplitscreenMode->GetInt());
+						ensureMsgf(false, TEXT("Selected r.ScreenshotTest.EmulateSplitscreenMode of %d is invalid. Defaulting to 0 (off)."), EmulateSplitscreenMode);
 						break;
 				}
 			}
